Move ConcertSyncUI from ConcertSyncClient plugin to ConcertSharedSlate plugin

diff --git a/Engine/Plugins/Developer/Concert/ConcertSync/ConcertSyncClient/ConcertSyncClient.uplugin b/Engine/Plugins/Developer/Concert/ConcertSync/ConcertSyncClient/ConcertSyncClient.uplugin
--- a/Engine/Plugins/Developer/Concert/ConcertSync/ConcertSyncClient/ConcertSyncClient.uplugin
+++ b/Engine/Plugins/Developer/Concert/ConcertSync/ConcertSyncClient/ConcertSyncClient.uplugin
@@ -20,11 +20,6 @@
 			"Name": "ConcertSyncClient",
 			"Type": "UncookedOnly",
 			"LoadingPhase": "Default"
-		},
-		{
-			"Name": "ConcertSyncUI",
-			"Type": "UncookedOnly",
-			"LoadingPhase": "Default"
 		}
 	],
 	"Plugins": [
diff --git a/Engine/Plugins/Developer/Concert/ConcertSync/ConcertSyncClient/Source/ConcertSyncUI/ConcertSyncUI.Build.cs b/Engine/Plugins/Developer/Concert/ConcertSync/ConcertSyncClient/Source/ConcertSyncUI/ConcertSyncUI.Build.cs
deleted file mode 100644
diff --git a/Engine/Plugins/Developer/Concert/ConcertSync/ConcertSyncClient/Source/ConcertSyncUI/Private/ConcertFrontendStyle.cpp b/Engine/Plugins/Developer/Concert/ConcertSync/ConcertSyncClient/Source/ConcertSyncUI/Private/ConcertFrontendStyle.cpp
deleted file mode 100644
diff --git a/Engine/Plugins/Developer/Concert/ConcertSync/ConcertSyncClient/Source/ConcertSyncUI/Private/ConcertSyncUIModule.cpp b/Engine/Plugins/Developer/Concert/ConcertSync/ConcertSyncClient/Source/ConcertSyncUI/Private/ConcertSyncUIModule.cpp
deleted file mode 100644
diff --git a/Engine/Plugins/Developer/Concert/ConcertSync/ConcertSyncClient/Source/ConcertSyncUI/Private/SConcertScrollBox.cpp b/Engine/Plugins/Developer/Concert/ConcertSync/ConcertSyncClient/Source/ConcertSyncUI/Private/SConcertScrollBox.cpp
deleted file mode 100644
diff --git a/Engine/Plugins/Developer/Concert/ConcertSync/ConcertSyncClient/Source/ConcertSyncUI/Private/SConcertSessionActivities.cpp b/Engine/Plugins/Developer/Concert/ConcertSync/ConcertSyncClient/Source/ConcertSyncUI/Private/SConcertSessionActivities.cpp
deleted file mode 100644
diff --git a/Engine/Plugins/Developer/Concert/ConcertSync/ConcertSyncClient/Source/ConcertSyncUI/Private/SConcertSessionRecovery.cpp b/Engine/Plugins/Developer/Concert/ConcertSync/ConcertSyncClient/Source/ConcertSyncUI/Private/SConcertSessionRecovery.cpp
deleted file mode 100644
diff --git a/Engine/Plugins/Developer/Concert/ConcertSync/ConcertSyncClient/Source/ConcertSyncUI/Private/SPackageDetails.cpp b/Engine/Plugins/Developer/Concert/ConcertSync/ConcertSyncClient/Source/ConcertSyncUI/Private/SPackageDetails.cpp
deleted file mode 100644
diff --git a/Engine/Plugins/Developer/Concert/ConcertSync/ConcertSyncClient/Source/ConcertSyncUI/Public/ConcertFrontendStyle.h b/Engine/Plugins/Developer/Concert/ConcertSync/ConcertSyncClient/Source/ConcertSyncUI/Public/ConcertFrontendStyle.h
deleted file mode 100644
diff --git a/Engine/Plugins/Developer/Concert/ConcertSync/ConcertSyncClient/Source/ConcertSyncUI/Public/ConcertFrontendUtils.h b/Engine/Plugins/Developer/Concert/ConcertSync/ConcertSyncClient/Source/ConcertSyncUI/Public/ConcertFrontendUtils.h
deleted file mode 100644
diff --git a/Engine/Plugins/Developer/Concert/ConcertSync/ConcertSyncClient/Source/ConcertSyncUI/Public/IConcertSyncUIModule.h b/Engine/Plugins/Developer/Concert/ConcertSync/ConcertSyncClient/Source/ConcertSyncUI/Public/IConcertSyncUIModule.h
deleted file mode 100644
diff --git a/Engine/Plugins/Developer/Concert/ConcertSync/ConcertSyncClient/Source/ConcertSyncUI/Public/SConcertScrollBox.h b/Engine/Plugins/Developer/Concert/ConcertSync/ConcertSyncClient/Source/ConcertSyncUI/Public/SConcertScrollBox.h
deleted file mode 100644
diff --git a/Engine/Plugins/Developer/Concert/ConcertSync/ConcertSyncClient/Source/ConcertSyncUI/Public/SConcertSessionActivities.h b/Engine/Plugins/Developer/Concert/ConcertSync/ConcertSyncClient/Source/ConcertSyncUI/Public/SConcertSessionActivities.h
deleted file mode 100644
diff --git a/Engine/Plugins/Developer/Concert/ConcertSync/ConcertSyncClient/Source/ConcertSyncUI/Public/SConcertSessionRecovery.h b/Engine/Plugins/Developer/Concert/ConcertSync/ConcertSyncClient/Source/ConcertSyncUI/Public/SConcertSessionRecovery.h
deleted file mode 100644
diff --git a/Engine/Plugins/Developer/Concert/ConcertSync/ConcertSyncClient/Source/ConcertSyncUI/Public/SPackageDetails.h b/Engine/Plugins/Developer/Concert/ConcertSync/ConcertSyncClient/Source/ConcertSyncUI/Public/SPackageDetails.h
deleted file mode 100644
diff --git a/Engine/Plugins/Developer/Concert/ConcertUI/ConcertSharedSlate/ConcertSharedSlate.uplugin b/Engine/Plugins/Developer/Concert/ConcertUI/ConcertSharedSlate/ConcertSharedSlate.uplugin
new file mode 100644
--- /dev/null
+++ b/Engine/Plugins/Developer/Concert/ConcertUI/ConcertSharedSlate/ConcertSharedSlate.uplugin
@@ -0,0 +1,45 @@
+{
+	"FileVersion": 3,
+	"Version": 1,
+	"VersionName": "0.2",
+	"FriendlyName": "Concert Sync Server UI",
+	"Description": "UI plugin that visualizes the multi-user server",
+	"Category": "Networking",
+	"CreatedBy": "Epic Games, Inc.",
+	"CreatedByURL": "http://epicgames.com",
+	"DocsURL": "",
+	"MarketplaceURL": "",
+	"SupportURL": "",
+	"EnabledByDefault": false,
+	"CanContainContent": false,
+	"Hidden": true,
+	"IsBetaVersion": true,
+	"Installed": false,
+	"SupportedPrograms": [ "UnrealMultiUserServer", "UnrealRecoverySvc", "CrashReportClientEditor" ],
+	"Modules": [
+        {
+          "Name": "ConcertSharedSlate",
+          "Type": "UncookedOnly",
+          "LoadingPhase": "Default",
+          "ProgramAllowList": [ "UnrealMultiUserServer", "UnrealRecoverySvc", "CrashReportClientEditor" ]
+        }
+	],
+	"Plugins": [
+		{
+			"Name": "ConcertMain",
+			"Enabled": true
+		},
+		{
+			"Name": "ConcertSyncCore",
+			"Enabled": true
+		},
+        {
+          "Name": "ConcertSyncClient",
+          "Enabled": true
+        },
+		{
+			"Name": "ConcertSyncServer",
+			"Enabled": true
+		}
+	]
+}
diff --git a/Engine/Plugins/Developer/Concert/ConcertUI/ConcertSharedSlate/Source/ConcertSharedSlate/ConcertSharedSlate.Build.cs b/Engine/Plugins/Developer/Concert/ConcertUI/ConcertSharedSlate/Source/ConcertSharedSlate/ConcertSharedSlate.Build.cs
new file mode 100644
--- /dev/null
+++ b/Engine/Plugins/Developer/Concert/ConcertUI/ConcertSharedSlate/Source/ConcertSharedSlate/ConcertSharedSlate.Build.cs
@@ -0,0 +1,35 @@
+// Copyright Epic Games, Inc. All Rights Reserved.
+
+namespace UnrealBuildTool.Rules
+{
+	public class ConcertSharedSlate : ModuleRules
+	{
+		public ConcertSharedSlate(ReadOnlyTargetRules Target) : base(Target)
+		{
+			PublicDependencyModuleNames.AddRange(
+				new string[]
+				{
+					"Core",
+					"CoreUObject",
+					"Slate",
+					"EditorStyle",
+					"Concert"
+				}
+			);
+
+			PrivateDependencyModuleNames.AddRange(
+				new string[]
+				{
+					"SlateCore",
+					"InputCore",
+					"Projects",
+					"Concert",
+					"ConcertTransport",
+					"ConcertSyncClient",
+					"ConcertSyncCore",
+					"UndoHistory"
+				}
+			);
+		}
+	}
+}
diff --git a/Engine/Plugins/Developer/Concert/ConcertUI/ConcertSharedSlate/Source/ConcertSharedSlate/Private/ConcertFrontendStyle.cpp b/Engine/Plugins/Developer/Concert/ConcertUI/ConcertSharedSlate/Source/ConcertSharedSlate/Private/ConcertFrontendStyle.cpp
new file mode 100644
--- /dev/null
+++ b/Engine/Plugins/Developer/Concert/ConcertUI/ConcertSharedSlate/Source/ConcertSharedSlate/Private/ConcertFrontendStyle.cpp
@@ -0,0 +1,144 @@
+// Copyright Epic Games, Inc. All Rights Reserved.
+
+#include "ConcertFrontendStyle.h"
+#include "Styling/SlateStyleRegistry.h"
+#include "Styling/SlateTypes.h"
+#include "EditorStyleSet.h"
+#include "Interfaces/IPluginManager.h"
+#include "SlateOptMacros.h"
+
+#define IMAGE_PLUGIN_BRUSH( RelativePath, ... ) FSlateImageBrush( FConcertFrontendStyle::InContent( RelativePath, ".png" ), __VA_ARGS__ )
+#define IMAGE_BRUSH(RelativePath, ...) FSlateImageBrush(StyleSet->RootToContentDir(RelativePath, TEXT(".png")), __VA_ARGS__)
+#define IMAGE_PLUGIN_BRUSH_SVG( RelativePath, ... ) FSlateVectorImageBrush( FConcertFrontendStyle::InContent(RelativePath, ".svg"), __VA_ARGS__)
+
+FString FConcertFrontendStyle::InContent(const FString& RelativePath, const ANSICHAR* Extension)
+{
+	static FString ContentDir = IPluginManager::Get().FindPlugin(TEXT("ConcertSharedSlate"))->GetContentDir();
+	return (ContentDir / RelativePath) + Extension;
+}
+
+TSharedPtr< class FSlateStyleSet > FConcertFrontendStyle::StyleSet;
+
+FName FConcertFrontendStyle::GetStyleSetName()
+{
+	return FName(TEXT("ConcertFrontendStyle"));
+}
+
+void FConcertFrontendStyle::Initialize()
+{
+	// Only register once
+	if (StyleSet.IsValid())
+	{
+		return;
+	}
+
+	StyleSet = MakeShared<FSlateStyleSet>(GetStyleSetName());
+	StyleSet->SetContentRoot(FPaths::EngineContentDir() / TEXT("Editor/Slate"));
+	StyleSet->SetCoreContentRoot(FPaths::EngineContentDir() / TEXT("Slate"));
+
+	// Const icon sizes
+	const FVector2D Icon16x16(16.0f, 16.0f);
+	const FVector2D Icon20x20(20.0f, 20.0f);
+	const FVector2D Icon24x24(24.0f, 24.0f);
+	const FVector2D Icon40x40(40.0f, 40.0f);
+	const FVector2D Icon48x48(48.0f, 48.0f);
+
+	// Use this to change the opacity. Ex: In UE4, the icon looked fine at 80%, in UE5, icons looks homogeneous with others at 100%.
+	const FLinearColor IconColorAndOpacity(FLinearColor(1.f, 1.f, 1.f, 1.f));
+
+	// 16x16
+	StyleSet->Set("Concert.Persist",         new IMAGE_PLUGIN_BRUSH("Icons/icon_ConcertPersist_16x", Icon16x16));
+	StyleSet->Set("Concert.LockBackground",  new IMAGE_PLUGIN_BRUSH("Icons/icon_ConcertLockBackground_16x", Icon16x16));
+	StyleSet->Set("Concert.MyLock",          new IMAGE_PLUGIN_BRUSH("Icons/icon_ConcertMyLock_16x", Icon16x16));
+	StyleSet->Set("Concert.OtherLock",       new IMAGE_PLUGIN_BRUSH("Icons/icon_ConcertOtherLock_16x", Icon16x16));
+	StyleSet->Set("Concert.ModifiedByOther", new IMAGE_PLUGIN_BRUSH("Icons/icon_ConcertModifiedByOther_16x", Icon16x16));
+
+	// Multi-user Tab/Menu icons
+	StyleSet->Set("Concert.MultiUser", new IMAGE_PLUGIN_BRUSH_SVG("Icons/icon_MultiUser", Icon16x16));
+
+	// Maps the UI Command name in Multi-User module. (UI_COMMAND does magic icon mapping when style name and command name matches)
+	StyleSet->Set("Concert.OpenBrowser",  new IMAGE_PLUGIN_BRUSH("Icons/icon_MultiUser_32x", Icon16x16));
+	StyleSet->Set("Concert.OpenSettings", new IMAGE_PLUGIN_BRUSH("Icons/icon_Settings_32x",  Icon16x16, IconColorAndOpacity));
+	StyleSet->Set("Concert.LaunchServer", new IMAGE_PLUGIN_BRUSH("Icons/icon_NewServer_32x", Icon16x16, IconColorAndOpacity));
+
+	// Multi-User Browser
+	StyleSet->Set("Concert.ArchiveSession",     new IMAGE_PLUGIN_BRUSH("Icons/icon_ArchiveSession_48x",     Icon24x24, IconColorAndOpacity));
+	StyleSet->Set("Concert.CancelAutoJoin",     new IMAGE_PLUGIN_BRUSH("Icons/icon_CancelAutoJoin_48x",     Icon24x24, IconColorAndOpacity));
+	StyleSet->Set("Concert.CloseServer",        new IMAGE_PLUGIN_BRUSH("Icons/icon_CloseServer_48x",        Icon24x24, IconColorAndOpacity));
+	StyleSet->Set("Concert.DeleteSession",      new IMAGE_PLUGIN_BRUSH("Icons/icon_DeleteSession_48x",      Icon24x24, IconColorAndOpacity));
+	StyleSet->Set("Concert.JoinDefaultSession", new IMAGE_PLUGIN_BRUSH("Icons/icon_JoinDefaultSession_48x", Icon24x24, IconColorAndOpacity));
+	StyleSet->Set("Concert.JoinSession",        new IMAGE_PLUGIN_BRUSH("Icons/icon_JoinSelectedSession_48x",Icon24x24, IconColorAndOpacity));
+	StyleSet->Set("Concert.LeaveSession",       new IMAGE_PLUGIN_BRUSH("Icons/icon_LeaveSession_48x",       Icon24x24, IconColorAndOpacity));
+	StyleSet->Set("Concert.NewServer",          new IMAGE_PLUGIN_BRUSH("Icons/icon_NewServer_48x",          Icon24x24, IconColorAndOpacity));
+	StyleSet->Set("Concert.NewSession",         new IMAGE_PLUGIN_BRUSH("Icons/icon_NewSession_48x",         Icon24x24, IconColorAndOpacity));
+	StyleSet->Set("Concert.PauseSession",       new IMAGE_PLUGIN_BRUSH("Icons/icon_PauseSession_48x",       Icon24x24, IconColorAndOpacity));
+	StyleSet->Set("Concert.RestoreSession",     new IMAGE_PLUGIN_BRUSH("Icons/icon_RestoreSession_48x",     Icon24x24, IconColorAndOpacity));
+	StyleSet->Set("Concert.ResumeSession",      new IMAGE_PLUGIN_BRUSH("Icons/icon_ResumeSession_48x",      Icon24x24, IconColorAndOpacity));
+	StyleSet->Set("Concert.Settings",           new IMAGE_PLUGIN_BRUSH("Icons/icon_Settings_48x",           Icon24x24, IconColorAndOpacity));
+	StyleSet->Set("Concert.NewServer.Small",    new IMAGE_PLUGIN_BRUSH("Icons/icon_NewServer_32x",          Icon16x16, IconColorAndOpacity));
+	StyleSet->Set("Concert.NewSession.Small",   new IMAGE_PLUGIN_BRUSH("Icons/icon_NewSession_32x",         Icon16x16, IconColorAndOpacity));
+
+	// Multi-user Active session
+	StyleSet->Set("Concert.JumpToLocation",     new IMAGE_PLUGIN_BRUSH("Icons/icon_PresenceLocation_32x",   Icon16x16, IconColorAndOpacity));
+	StyleSet->Set("Concert.HidePresence",       new IMAGE_PLUGIN_BRUSH("Icons/icon_PresenceEyeOff_32x",     Icon16x16, IconColorAndOpacity));
+	StyleSet->Set("Concert.ShowPresence",       new IMAGE_PLUGIN_BRUSH("Icons/icon_PresenceEyeOn_32x",      Icon16x16, IconColorAndOpacity));
+
+	// 24x24/48x48 -> For sequencer toolbar.
+	StyleSet->Set("Concert.Sequencer.SyncTimeline",       new IMAGE_PLUGIN_BRUSH("Icons/icon_SequencerSyncTimeline_48x", Icon48x48, IconColorAndOpacity)); // Enable/disable playback and time scrubbing from a remote client.
+	StyleSet->Set("Concert.Sequencer.SyncTimeline.Small", new IMAGE_PLUGIN_BRUSH("Icons/icon_SequencerSyncTimeline_48x", Icon24x24, IconColorAndOpacity));
+	StyleSet->Set("Concert.Sequencer.SyncSequence",       new IMAGE_PLUGIN_BRUSH("Icons/icon_SequencerSyncSequence_48x", Icon48x48, IconColorAndOpacity)); // Allows or not a remote client to open/close sequencer.
+	StyleSet->Set("Concert.Sequencer.SyncSequence.Small", new IMAGE_PLUGIN_BRUSH("Icons/icon_SequencerSyncSequence_48x", Icon24x24, IconColorAndOpacity));
+	StyleSet->Set("Concert.Sequencer.SyncUnrelated",      new IMAGE_PLUGIN_BRUSH("Icons/icon_SequencerSyncUnrelated_48x", Icon48x48, IconColorAndOpacity)); // Enable/disable playback and time scrubbing from a remote client event if this user has a different sequence opened.
+	StyleSet->Set("Concert.Sequencer.SyncUnrelated.Small",new IMAGE_PLUGIN_BRUSH("Icons/icon_SequencerSyncUnrelated_48x", Icon24x24, IconColorAndOpacity));
+
+	// 40x40 -> Editor toolbar large icons.
+	StyleSet->Set("Concert.Browse", new IMAGE_PLUGIN_BRUSH("Icons/icon_MultiUserMenuBrowse_40x", Icon40x40));
+	StyleSet->Set("Concert.Join",   new IMAGE_PLUGIN_BRUSH("Icons/icon_MultiUserMenuJoin_40x",   Icon40x40));
+	StyleSet->Set("Concert.Leave",  new IMAGE_PLUGIN_BRUSH("Icons/icon_MultiUserMenuLeave_40x",  Icon40x40));
+	StyleSet->Set("Concert.Cancel", new IMAGE_PLUGIN_BRUSH("Icons/icon_MultiUserMenuCancel_40x", Icon40x40));
+
+	// 20x20 -> Editor toolbar small icons.
+	StyleSet->Set("Concert.Browse.Small", new IMAGE_PLUGIN_BRUSH("Icons/icon_MultiUserMenuBrowse_40x", Icon20x20));
+	StyleSet->Set("Concert.Leave.Small",  new IMAGE_PLUGIN_BRUSH("Icons/icon_MultiUserMenuLeave_40x",  Icon20x20));
+	StyleSet->Set("Concert.Join.Small",   new IMAGE_PLUGIN_BRUSH("Icons/icon_MultiUserMenuJoin_40x",   Icon20x20));
+	StyleSet->Set("Concert.Cancel.Small", new IMAGE_PLUGIN_BRUSH("Icons/icon_MultiUserMenuCancel_40x", Icon20x20));
+
+	// Disaster Recovery
+	StyleSet->Set("Concert.RecoveryHub", new IMAGE_PLUGIN_BRUSH("Icons/icon_RecoveryHub_16x", Icon16x16));
+
+	// Activity Text
+	{
+		FTextBlockStyle BoldText = FEditorStyle::Get().GetWidgetStyle<FTextBlockStyle>("RichTextBlock.Bold");
+		StyleSet->Set("ActivityText.Bold", FTextBlockStyle(BoldText));
+	}
+
+	// Colors
+	{
+		StyleSet->Set("Concert.Color.LocalUser", FLinearColor(0.31f, 0.749f, 0.333f));
+		StyleSet->Set("Concert.Color.OtherUser", FLinearColor(0.93f, 0.608f, 0.169f));
+	}
+
+	// Colors
+	StyleSet->Set("Concert.DisconnectedColor", FLinearColor(0.672f, 0.672f, 0.672f));
+
+	FSlateStyleRegistry::RegisterSlateStyle(*StyleSet.Get());
+};
+
+void FConcertFrontendStyle::Shutdown()
+{
+	if (StyleSet.IsValid())
+	{
+		FSlateStyleRegistry::UnRegisterSlateStyle(*StyleSet.Get());
+		ensure(StyleSet.IsUnique());
+		StyleSet.Reset();
+	}
+}
+
+TSharedPtr<class ISlateStyle> FConcertFrontendStyle::Get()
+{
+	return StyleSet;
+}
+
+#undef IMAGE_PLUGIN_BRUSH
+#undef IMAGE_BRUSH
+
diff --git a/Engine/Plugins/Developer/Concert/ConcertUI/ConcertSharedSlate/Source/ConcertSharedSlate/Private/ConcertSharedSlateModule.cpp b/Engine/Plugins/Developer/Concert/ConcertUI/ConcertSharedSlate/Source/ConcertSharedSlate/Private/ConcertSharedSlateModule.cpp
new file mode 100644
--- /dev/null
+++ b/Engine/Plugins/Developer/Concert/ConcertUI/ConcertSharedSlate/Source/ConcertSharedSlate/Private/ConcertSharedSlateModule.cpp
@@ -0,0 +1,22 @@
+// Copyright Epic Games, Inc. All Rights Reserved.
+
+#include "IConcertSharedSlateModule.h"
+
+class FConcertSharedSlateModule : public IConcertSharedSlateModule
+{
+public:
+	FConcertSharedSlateModule() = default;
+	virtual ~FConcertSharedSlateModule() {}
+
+	virtual void StartupModule() override
+	{
+	}
+
+	virtual void ShutdownModule() override
+	{
+	}
+};
+
+
+IMPLEMENT_MODULE(FConcertSharedSlateModule, ConcertSharedSlate);
+
diff --git a/Engine/Plugins/Developer/Concert/ConcertUI/ConcertSharedSlate/Source/ConcertSharedSlate/Private/SConcertScrollBox.cpp b/Engine/Plugins/Developer/Concert/ConcertUI/ConcertSharedSlate/Source/ConcertSharedSlate/Private/SConcertScrollBox.cpp
new file mode 100644
--- /dev/null
+++ b/Engine/Plugins/Developer/Concert/ConcertUI/ConcertSharedSlate/Source/ConcertSharedSlate/Private/SConcertScrollBox.cpp
@@ -0,0 +1,148 @@
+// Copyright Epic Games, Inc. All Rights Reserved.
+
+#include "SConcertScrollBox.h"
+#include "Widgets/Layout/SScrollBox.h"
+#include "EditorStyleSet.h"
+#include "Widgets/Input/SButton.h"
+#include "Widgets/Text/STextBlock.h"
+#include "Widgets/Layout/SBox.h"
+#include "Widgets/SBoxPanel.h"
+
+#define LOCTEXT_NAMESPACE "SConcertScrollBox"
+
+void SConcertScrollBox::Construct(const FArguments& InArgs)
+{
+	bIsLocked = true;
+	bPreventLock = false;
+
+	ScrollBar = SNew(SScrollBar)
+		.Thickness(FVector2D(12.0f, 12.0f));
+
+	ChildSlot
+		[
+			SNew(SHorizontalBox)
+			+SHorizontalBox::Slot()	
+			[
+ 				SAssignNew(ScrollBox, SScrollBox)
+ 				.ExternalScrollbar(ScrollBar)
+				.OnUserScrolled(this, &SConcertScrollBox::HandleUserScrolled)
+			]
+		
+			+SHorizontalBox::Slot()
+			.AutoWidth()
+			[
+				SNew(SVerticalBox)
+
+				+SVerticalBox::Slot()
+				.VAlign(VAlign_Bottom)
+				.AutoHeight()
+				[
+					CreateScrollBarButton(LOCTEXT("ScrollToStartToolTip", "Scroll to the start of the list."), TEXT("\xf077"), FOnClicked::CreateSP(this, &SConcertScrollBox::HandleScrollToStart)) 
+				]
+
+				+SVerticalBox::Slot()
+				.FillHeight(1.f)
+				[
+					ScrollBar.ToSharedRef()
+				]
+
+				+SVerticalBox::Slot()
+				.VAlign(VAlign_Bottom)
+				.AutoHeight()
+				[
+					CreateScrollBarButton(LOCTEXT("ScrollToEndToolTip", "Scroll to the end of the list."), TEXT("\xf078"), FOnClicked::CreateSP(this, &SConcertScrollBox::HandleScrollToEnd))
+				]
+			]
+		];
+
+	// Forward the slots to the inner scroll box.
+	for (const SConcertScrollBox::FSlot::FSlotArguments& InSlot : InArgs._Slots)
+	{
+		if (InSlot.GetAttachedWidget())
+		{
+			ScrollBox->AddSlot()
+			[
+				InSlot.GetAttachedWidget().ToSharedRef()
+			];
+		}
+	}
+}
+
+void SConcertScrollBox::Tick(const FGeometry& AllottedGeometry, const double InCurrentTime, const float InDeltaTime)
+{
+	if (IsAtBottom() && !bPreventLock)
+	{
+		bIsLocked = true;
+	}
+
+	if (bIsLocked)
+	{
+		ScrollBox->ScrollToEnd();
+	};
+
+	/** Locking was prevented for one tick to allow the scrollbar to move without triggering a new lock. */
+	if (bPreventLock)
+	{
+		bPreventLock = false;
+	}
+}
+
+
+SConcertScrollBox::FSlot::FSlotArguments SConcertScrollBox::Slot()
+{
+	return SConcertScrollBox::FSlot::FSlotArguments(MakeUnique<SConcertScrollBox::FSlot>());
+}
+
+void SConcertScrollBox::HandleUserScrolled(float)
+{
+	bPreventLock = true;
+	bIsLocked = false;
+}
+
+bool SConcertScrollBox::IsAtBottom()
+{
+	return ScrollBar->DistanceFromBottom() == 0.f;
+}
+
+FReply SConcertScrollBox::HandleScrollToStart()
+{
+	bIsLocked = false;
+	bPreventLock = true;
+	ScrollBox->ScrollToStart();
+	return FReply::Handled();
+}
+
+FReply SConcertScrollBox::HandleScrollToEnd()
+{
+	bIsLocked = true;
+	ScrollBox->ScrollToEnd();
+	return FReply::Handled();
+}
+
+TSharedRef<SButton> SConcertScrollBox::CreateScrollBarButton(const FText& InToolTip, const FString& InIcon, FOnClicked OnClickedDelegate) const
+{
+	return SNew(SButton)
+		.Visibility(this, &SConcertScrollBox::HandleScrollButtonsVisibility)
+		.ToolTipText(InToolTip)
+		.ButtonStyle(FEditorStyle::Get(), "FlatButton.Default")
+		.ForegroundColor(FLinearColor::White)
+		.OnClicked(OnClickedDelegate)
+		.ContentPadding(FMargin(2.f, 2.f))
+		[
+			SNew(SBox)
+			.HAlign(HAlign_Center)
+			[
+				SNew(STextBlock)
+				.Font(FEditorStyle::Get().GetFontStyle("FontAwesome.9"))
+				.Text(FText::FromString(InIcon)) /*fa-chevron-up*/
+				.Justification(ETextJustify::Center)
+			]
+		];
+}
+
+EVisibility SConcertScrollBox::HandleScrollButtonsVisibility() const
+{
+	return ScrollBar->GetVisibility();
+}
+
+#undef LOCTEXT_NAMESPACE /** SConcertScrollBox */
\ No newline at end of file
diff --git a/Engine/Plugins/Developer/Concert/ConcertUI/ConcertSharedSlate/Source/ConcertSharedSlate/Private/SConcertSessionActivities.cpp b/Engine/Plugins/Developer/Concert/ConcertUI/ConcertSharedSlate/Source/ConcertSharedSlate/Private/SConcertSessionActivities.cpp
new file mode 100644
--- /dev/null
+++ b/Engine/Plugins/Developer/Concert/ConcertUI/ConcertSharedSlate/Source/ConcertSharedSlate/Private/SConcertSessionActivities.cpp
@@ -0,0 +1,1176 @@
+// Copyright Epic Games, Inc. All Rights Reserved.
+
+#include "SConcertSessionActivities.h"
+
+#include "IConcertClientWorkspace.h"
+#include "Algo/Transform.h"
+#include "Misc/AsyncTaskNotification.h"
+#include "Misc/ITransaction.h"
+#include "Framework/Application/SlateApplication.h"
+#include "Framework/MultiBox/MultiBoxBuilder.h"
+#include "Widgets/SOverlay.h"
+#include "Widgets/SUndoHistoryDetails.h"
+#include "Widgets/Colors/SColorBlock.h"
+#include "Widgets/Input/SComboButton.h"
+#include "Widgets/Images/SThrobber.h"
+#include "Widgets/Layout/SBorder.h"
+#include "Widgets/Layout/SBox.h"
+#include "Widgets/Layout/SExpandableArea.h"
+#include "Widgets/Layout/SScrollBox.h"
+#include "Widgets/Layout/SSpacer.h"
+#include "Widgets/Layout/SSplitter.h"
+#include "Widgets/Text/SRichTextBlock.h"
+#include "Widgets/Text/STextBlock.h"
+#include "Widgets/Views/STableRow.h"
+#include "Styling/AppStyle.h"
+#include "ConcertFrontendStyle.h"
+#include "ConcertFrontendUtils.h"
+#include "ConcertLogGlobal.h"
+#include "ConcertTransactionEvents.h"
+#include "ConcertWorkspaceData.h"
+#include "SPackageDetails.h"
+
+#define LOCTEXT_NAMESPACE "SConcertSessionActivities"
+
+namespace ConcertSessionActivityUtils
+{
+// The columns names.
+const FName DateTimeColumnId    = TEXT("DateTime");
+const FName OperationColumnId   = TEXT("Operation");
+const FName PackageColumnId     = TEXT("Package");
+const FName SummaryColumnId     = TEXT("Summary");
+const FName ClientNameColumnId  = TEXT("Client");
+const FName AvatarColorColumnId = TEXT("Client AvatarColor");
+
+// The View Options check boxes.
+const FName DisplayRelativeTimeCheckBoxId       = TEXT("DisplayRelativeTime");
+const FName ShowConnectionActivitiesCheckBoxId  = TEXT("ShowConnectionActivities");
+const FName ShowLockActivitiesCheckBoxId        = TEXT("ShowLockActivities");
+const FName ShowPackageActivitiesCheckBoxId     = TEXT("ShowPackageActivities");
+const FName ShowTransactionActivitiesCheckBoxId = TEXT("ShowTransactionActivities");
+const FName ShowIgnoredActivitiesCheckBoxId     = TEXT("ShowIgnoredActivities");
+
+FText GetActivityDateTime(const FConcertClientSessionActivity& Activity, SConcertSessionActivities::ETimeFormat TimeFormat)
+{
+	return TimeFormat == SConcertSessionActivities::ETimeFormat::Relative ? ConcertFrontendUtils::FormatRelativeTime(Activity.Activity.EventTime) : FText::AsDateTime(Activity.Activity.EventTime);
+}
+
+FText GetOperationName(const FConcertClientSessionActivity& Activity)
+{
+	if (const FConcertSyncTransactionActivitySummary* TransactionSummary = Activity.ActivitySummary.Cast<FConcertSyncTransactionActivitySummary>())
+	{
+		return TransactionSummary->TransactionTitle;
+	}
+
+	if (const FConcertSyncPackageActivitySummary* PackageSummary = Activity.ActivitySummary.Cast<FConcertSyncPackageActivitySummary>())
+	{
+		auto GetSavedOperationNameFn = [](const FConcertSyncPackageActivitySummary* PackageSummary)
+		{
+			if (PackageSummary->bAutoSave)
+			{
+				return LOCTEXT("AutoSavePackageOperation", "Auto-Save Package");
+			}
+			return PackageSummary->bPreSave ? LOCTEXT("PreSavePackageOperation", "Pre-Save Package") : LOCTEXT("SavePackageOperation", "Save Package");
+		};
+
+		switch (PackageSummary->PackageUpdateType)
+		{
+		case EConcertPackageUpdateType::Added   : return LOCTEXT("NewPackageOperation",    "New Package");
+		case EConcertPackageUpdateType::Deleted : return LOCTEXT("DeletePackageOperation", "Delete Package");
+		case EConcertPackageUpdateType::Renamed : return LOCTEXT("RenamePackageOperation", "Rename Package");
+		case EConcertPackageUpdateType::Saved   : return GetSavedOperationNameFn(PackageSummary);
+		case EConcertPackageUpdateType::Dummy   : return LOCTEXT("DiscardPackageOperation", "Discard Changes");
+		default: break;
+		}
+	}
+
+	if (const FConcertSyncConnectionActivitySummary* ConnectionSummary = Activity.ActivitySummary.Cast<FConcertSyncConnectionActivitySummary>())
+	{
+		switch (ConnectionSummary->ConnectionEventType)
+		{
+		case EConcertSyncConnectionEventType::Connected:    return LOCTEXT("JoinOperation", "Join Session");
+		case EConcertSyncConnectionEventType::Disconnected: return LOCTEXT("LeaveOperation", "Leave Session");
+		default: break;
+		}
+	}
+
+	if (const FConcertSyncLockActivitySummary* LockSummary = Activity.ActivitySummary.Cast<FConcertSyncLockActivitySummary>())
+	{
+		switch (LockSummary->LockEventType)
+		{
+		case EConcertSyncLockEventType::Locked:   return LOCTEXT("LockOperation", "Lock");
+		case EConcertSyncLockEventType::Unlocked: return LOCTEXT("UnlockOperation", "Unlock");
+		default: break;
+		}
+	}
+
+	return FText::GetEmpty();
+}
+
+FText GetPackageName(const FConcertClientSessionActivity& Activity)
+{
+	if (const FConcertSyncPackageActivitySummary* Summary = Activity.ActivitySummary.Cast<FConcertSyncPackageActivitySummary>())
+	{
+		return FText::FromName(Summary->PackageName);
+	}
+
+	if (const FConcertSyncTransactionActivitySummary* Summary = Activity.ActivitySummary.Cast<FConcertSyncTransactionActivitySummary>())
+	{
+		return FText::FromName(Summary->PrimaryPackageName);
+	}
+
+	return FText::GetEmpty();
+}
+
+FText GetSummary(const FConcertClientSessionActivity& Activity, const FText& ClientName, bool bAsRichText)
+{
+	if (const FConcertSyncActivitySummary* Summary = Activity.ActivitySummary.Cast<FConcertSyncActivitySummary>())
+	{
+		return Summary->ToDisplayText(ClientName, bAsRichText);
+	}
+
+	return FText::GetEmpty();
+}
+
+FText GetClientName(const FConcertClientInfo* InActivityClient)
+{
+	return InActivityClient ? FText::AsCultureInvariant(InActivityClient->DisplayName) : FText::GetEmpty();
+}
+
+};
+
+/**
+ * Displays the summary of an activity recorded and recoverable in the SConcertSessionRecovery list view.
+ */
+class SConcertSessionActivityRow : public SMultiColumnTableRow<TSharedPtr<FConcertClientSessionActivity>>
+{
+public:
+	SLATE_BEGIN_ARGS(SConcertSessionActivityRow)
+		: _TimeFormat(SConcertSessionActivities::ETimeFormat::Relative)
+		, _HighlightText()
+		, _OnMakeColumnOverlayWidget()
+	{
+	}
+
+	SLATE_ATTRIBUTE(SConcertSessionActivities::ETimeFormat, TimeFormat)
+	SLATE_ATTRIBUTE(FText, HighlightText)
+	SLATE_ARGUMENT(SConcertSessionActivities::FMakeColumnOverlayWidgetFunc, OnMakeColumnOverlayWidget) // Function invoked when generating a row to add a widget above the column widget.
+
+	SLATE_END_ARGS()
+
+public:
+	/**
+	 * Constructs a row widget to display a Concert activity.
+	 * @param InArgs The widgets arguments.
+	 * @param InActivty The activity to display.
+	 * @param InActivityClient The client who produced this activity. Can be null if unknown or not desirable.
+	 * @param InOwnerTableView The table view that will own this row.
+	 */
+	void Construct(const FArguments& InArgs, TSharedPtr<FConcertClientSessionActivity> InActivity, const FConcertClientInfo* InActivityClient, const TSharedRef<STableViewBase>& InOwnerTableView);
+
+	/** Generates the widget representing this row. */
+	virtual TSharedRef<SWidget> GenerateWidgetForColumn(const FName& ColumnName) override;
+
+private:
+	/** Format the the time displayed as relative or absolute according to the attribute 'TimeFormat'. */
+	FText FormatEventDateTime() const;
+
+	/** Returns the client avatar color. */
+	FLinearColor GetClientAvatarColor() const { return ClientAvatarColor; }
+
+	/** Returns the display name of the client who performed the activity or an empty text if the information wasn't available or desired. */
+	FText GetClientName() const { return ClientName; }
+
+	/** Generates the tooltip for this row. */
+	FText MakeTooltipText() const;
+
+private:
+	TWeakPtr<FConcertClientSessionActivity> Activity;
+	TAttribute<SConcertSessionActivities::ETimeFormat> TimeFormat;
+	FText AbsoluteDateTime;
+	FText ClientName;
+	FLinearColor ClientAvatarColor;
+	TAttribute<FText> HighlightText;
+	SConcertSessionActivities::FMakeColumnOverlayWidgetFunc OnMakeColumnOverlayWidget;
+};
+
+
+void SConcertSessionActivityRow::Construct(const FArguments& InArgs, TSharedPtr<FConcertClientSessionActivity> InActivity, const FConcertClientInfo* InActivityClient, const TSharedRef<STableViewBase>& InOwnerTableView)
+{
+	Activity = InActivity;
+	TimeFormat = InArgs._TimeFormat;
+	HighlightText = InArgs._HighlightText;
+	OnMakeColumnOverlayWidget = InArgs._OnMakeColumnOverlayWidget;
+	AbsoluteDateTime = ConcertSessionActivityUtils::GetActivityDateTime(*InActivity, SConcertSessionActivities::ETimeFormat::Absolute); // Cache the absolute time. It doesn't changes.
+	ClientName = ConcertSessionActivityUtils::GetClientName(InActivityClient);
+	ClientAvatarColor = InActivityClient ? InActivityClient->AvatarColor : FConcertFrontendStyle::Get()->GetColor("Concert.DisconnectedColor");
+
+	// Construct base class
+	SMultiColumnTableRow<TSharedPtr<FConcertClientSessionActivity>>::Construct(FSuperRowType::FArguments(), InOwnerTableView);
+
+	if (InActivity->Activity.bIgnored)
+	{
+		SetColorAndOpacity(FLinearColor(0.5, 0.5, 0.5, 0.5));
+	}
+}
+
+TSharedRef<SWidget> SConcertSessionActivityRow::GenerateWidgetForColumn(const FName& ColumnId)
+{
+	TSharedPtr<FConcertClientSessionActivity> ActivityPin = Activity.Pin();
+	TSharedRef<SOverlay> Overlay = SNew(SOverlay);
+
+	if (ColumnId == ConcertSessionActivityUtils::AvatarColorColumnId)
+	{
+		Overlay->AddSlot()
+		.HAlign(HAlign_Center)
+		.VAlign(VAlign_Center)
+		.Padding(2, 1)
+		[
+			SNew(SColorBlock)
+			.Color(GetClientAvatarColor())
+			.Size(FVector2D(4.f, 16.f))
+		];
+	}
+	else if (ColumnId == ConcertSessionActivityUtils::DateTimeColumnId)
+	{
+		Overlay->AddSlot()
+		.VAlign(VAlign_Center)
+		[
+			SNew(STextBlock)
+			.Text(this, &SConcertSessionActivityRow::FormatEventDateTime)
+			.HighlightText(HighlightText)
+		];
+	}
+	else if (ColumnId == ConcertSessionActivityUtils::ClientNameColumnId)
+	{
+		Overlay->AddSlot()
+		.VAlign(VAlign_Center)
+		[
+			SNew(STextBlock)
+			.Text(GetClientName())
+			.HighlightText(HighlightText)
+		];
+	}
+	else if (ColumnId == ConcertSessionActivityUtils::PackageColumnId)
+	{
+		Overlay->AddSlot()
+		.VAlign(VAlign_Center)
+		[
+			SNew(STextBlock)
+			.Text(ConcertSessionActivityUtils::GetPackageName(*ActivityPin))
+			.HighlightText(HighlightText)
+		];
+	}
+	else if (ColumnId == ConcertSessionActivityUtils::OperationColumnId)
+	{
+		Overlay->AddSlot()
+		.VAlign(VAlign_Center)
+		[
+			SNew(STextBlock)
+			.Text(ConcertSessionActivityUtils::GetOperationName(*ActivityPin))
+			.HighlightText(HighlightText)
+		];
+	}
+	else
+	{
+		check(ColumnId == ConcertSessionActivityUtils::SummaryColumnId);
+
+		Overlay->AddSlot()
+		.VAlign(VAlign_Center)
+		[
+			SNew(SRichTextBlock)
+			.DecoratorStyleSet(FConcertFrontendStyle::Get().Get())
+			.Text(ConcertSessionActivityUtils::GetSummary(*ActivityPin, FText::GetEmpty(), true/*bAsRichText*/))
+			.HighlightText(HighlightText)
+		];
+	}
+
+	if (OnMakeColumnOverlayWidget)
+	{
+		if (TSharedPtr<SWidget> OverlayedWidget = OnMakeColumnOverlayWidget(ActivityPin, ColumnId))
+		{
+			Overlay->AddSlot()
+			[
+				OverlayedWidget.ToSharedRef()
+			];
+		}
+	}
+
+	SetToolTipText(TAttribute<FText>(this, &SConcertSessionActivityRow::MakeTooltipText));
+
+	return Overlay;
+}
+
+FText SConcertSessionActivityRow::MakeTooltipText() const
+{
+	if (TSharedPtr<FConcertClientSessionActivity> ActivityPin = Activity.Pin())
+	{
+		FText Client = GetClientName();
+		FText Operation = ConcertSessionActivityUtils::GetOperationName(*ActivityPin);
+		FText Package = ConcertSessionActivityUtils::GetPackageName(*ActivityPin);
+		FText Summary = ConcertSessionActivityUtils::GetSummary(*ActivityPin, Client, /*bAsRichText*/false);
+
+		FTextBuilder TextBuilder;
+
+		if (!Operation.IsEmpty())
+		{
+			TextBuilder.AppendLine(Operation);
+		}
+
+		TextBuilder.AppendLineFormat(LOCTEXT("ActivityRowTooltip_DateTime", "{0} ({1})"), AbsoluteDateTime, ConcertFrontendUtils::FormatRelativeTime(ActivityPin->Activity.EventTime));
+
+		if (!Package.IsEmpty())
+		{
+			TextBuilder.AppendLine(Package);
+		}
+
+		if (!Summary.IsEmpty())
+		{
+			TextBuilder.AppendLine(Summary);
+		}
+
+		if (ActivityPin->Activity.bIgnored)
+		{
+			TextBuilder.AppendLine();
+			TextBuilder.AppendLine(LOCTEXT("IgnoredActivity", "** This activity cannot be recovered (likely recorded during a Multi-User session). It is displayed for crash inspection only. It will be ignored on restore."));
+		}
+
+		return TextBuilder.ToText();
+	}
+
+	return FText::GetEmpty();
+}
+
+FText SConcertSessionActivityRow::FormatEventDateTime() const
+{
+	if (TSharedPtr<FConcertClientSessionActivity> ItemPin = Activity.Pin())
+	{
+		return TimeFormat.Get() == SConcertSessionActivities::ETimeFormat::Relative ? ConcertFrontendUtils::FormatRelativeTime(ItemPin->Activity.EventTime) : AbsoluteDateTime;
+	}
+	return FText::GetEmpty();
+}
+
+
+void SConcertSessionActivities::Construct(const FArguments& InArgs)
+{
+	FetchActivitiesFn = InArgs._OnFetchActivities;
+	GetActivityUserFn = InArgs._OnMapActivityToClient;
+	GetTransactionEventFn = InArgs._OnGetTransactionEvent;
+	GetPackageEventFn = InArgs._OnGetPackageEvent;
+	MakeColumnOverlayWidgetFn = InArgs._OnMakeColumnOverlayWidget;
+	HighlightText = InArgs._HighlightText;
+	TimeFormat = InArgs._TimeFormat;
+	ClientNameColumnVisibility = InArgs._ClientNameColumnVisibility;
+	OperationColumnVisibility = InArgs._OperationColumnVisibility;
+	PackageColumnVisibility = InArgs._PackageColumnVisibility;
+	ConnectionActivitiesVisibility = InArgs._ConnectionActivitiesVisibility;
+	LockActivitiesVisibility = InArgs._LockActivitiesVisibility;
+	PackageActivitiesVisibility = InArgs._PackageActivitiesVisibility;
+	TransactionActivitiesVisibility = InArgs._TransactionActivitiesVisibility;
+	IgnoredActivitiesVisibility = InArgs._IgnoredActivitiesVisibility;
+	DetailsAreaVisibility = InArgs._DetailsAreaVisibility;
+	bAutoScrollDesired = InArgs._IsAutoScrollEnabled;
+
+	SearchTextFilter = MakeShared<TTextFilter<const FConcertClientSessionActivity&>>(TTextFilter<const FConcertClientSessionActivity&>::FItemToStringArray::CreateSP(this, &SConcertSessionActivities::PopulateSearchStrings));
+	SearchTextFilter->OnChanged().AddSP(this, &SConcertSessionActivities::OnActivityFilterUpdated);
+
+	// Set the initial filter state.
+	ActiveFilterFlags = QueryActiveActivityFilters();
+
+	// Create the table header. (Setting visibility on the column itself doesn't show/hide the column as one would expect, unfortunately)
+	TSharedRef<SHeaderRow> HeaderRow = SNew(SHeaderRow);
+	if (InArgs._ClientAvatarColorColumnVisibility.Get() == EVisibility::Visible)
+	{
+		HeaderRow->AddColumn(SHeaderRow::Column(ConcertSessionActivityUtils::AvatarColorColumnId)
+			.DefaultLabel(INVTEXT(""))
+			.ManualWidth(8));
+	}
+
+	HeaderRow->AddColumn(SHeaderRow::Column(ConcertSessionActivityUtils::DateTimeColumnId)
+		.DefaultLabel(LOCTEXT("DateTime", "Date/Time"))
+		.ManualWidth(160));
+
+	if (InArgs._ClientNameColumnVisibility.Get() == EVisibility::Visible)
+	{
+		HeaderRow->AddColumn(SHeaderRow::Column(ConcertSessionActivityUtils::ClientNameColumnId)
+			.DefaultLabel(LOCTEXT("Client", "Client"))
+			.ManualWidth(80));
+	}
+
+	if (InArgs._OperationColumnVisibility.Get() == EVisibility::Visible)
+	{
+		HeaderRow->AddColumn(SHeaderRow::Column(ConcertSessionActivityUtils::OperationColumnId)
+			.DefaultLabel(LOCTEXT("Operation", "Operation"))
+			.ManualWidth(160));
+	}
+
+	if (InArgs._PackageColumnVisibility.Get() == EVisibility::Visible)
+	{
+		HeaderRow->AddColumn(SHeaderRow::Column(ConcertSessionActivityUtils::PackageColumnId)
+			.DefaultLabel(LOCTEXT("Package", "Package"))
+			.ManualWidth(200));
+	}
+
+	HeaderRow->AddColumn(SHeaderRow::Column(ConcertSessionActivityUtils::SummaryColumnId)
+		.DefaultLabel(LOCTEXT("Summary", "Summary")));
+
+	ChildSlot
+	[
+		SNew(SSplitter)
+		.Orientation(Orient_Vertical)
+
+		// Activity List
+		+SSplitter::Slot()
+		.Value(0.75)
+		[
+			SNew(SOverlay)
+
+			+SOverlay::Slot() // Activity list itself.
+			[
+				SNew(SBorder)
+				.BorderImage(FAppStyle::Get().GetBrush("ToolPanel.GroupBorder"))
+				.BorderBackgroundColor(FSlateColor(FLinearColor(0.6, 0.6, 0.6)))
+				.Padding(0)
+				[
+					SAssignNew(ActivityView, SListView<TSharedPtr<FConcertClientSessionActivity>>)
+					.ListItemsSource(&Activities)
+					.OnGenerateRow(this, &SConcertSessionActivities::OnGenerateActivityRowWidget)
+					.SelectionMode(ESelectionMode::Single)
+					.AllowOverscroll(EAllowOverscroll::No)
+					.OnListViewScrolled(this, &SConcertSessionActivities::OnListViewScrolled)
+					.OnSelectionChanged(this, &SConcertSessionActivities::OnListViewSelectionChanged)
+					.HeaderRow(HeaderRow)
+				]
+			]
+
+			+SOverlay::Slot() // Display a reason why no activities are shown.
+			.HAlign(HAlign_Center)
+			.VAlign(VAlign_Center)
+			[
+				SNew(STextBlock)
+				.Visibility_Lambda([TextAttr = InArgs._NoActivitiesReasonText](){ return TextAttr.Get().IsEmptyOrWhitespace() ? EVisibility::Collapsed : EVisibility::Visible; })
+				.Text(InArgs._NoActivitiesReasonText)
+				.Justification(ETextJustify::Center)
+			]
+		]
+
+		// Activity details.
+		+SSplitter::Slot()
+		.Value(0.25)
+		.SizeRule(TAttribute<SSplitter::ESizeRule>(this, &SConcertSessionActivities::GetDetailsAreaSizeRule))
+		[
+			SAssignNew(ExpandableDetails, SExpandableArea)
+			.Visibility(GetDetailAreaVisibility())
+			.InitiallyCollapsed(true)
+			.BorderBackgroundColor(FLinearColor(0.6f, 0.6f, 0.6f, 1.0f))
+			.BorderImage_Lambda([this]() { return ConcertFrontendUtils::GetExpandableAreaBorderImage(*ExpandableDetails); })
+			.BodyBorderImage(FAppStyle::Get().GetBrush("ToolPanel.GroupBorder"))
+			.BodyBorderBackgroundColor(FLinearColor::White)
+			.OnAreaExpansionChanged(this, &SConcertSessionActivities::OnDetailsAreaExpansionChanged)
+			.Padding(0.0f)
+			.HeaderContent()
+			[
+				SNew(STextBlock)
+				.Text(LOCTEXT("Details", "Details"))
+				.Font(FAppStyle::Get().GetFontStyle("DetailsView.CategoryFontStyle"))
+				.ShadowOffset(FVector2D(1.0f, 1.0f))
+			]
+			.BodyContent()
+			[
+				SNew(SOverlay)
+
+				+SOverlay::Slot()
+				[
+					SNew(SScrollBox)
+					.ScrollBarThickness(FVector2D(12.0f, 5.0f)) // To have same thickness than the ListView scroll bar.
+
+					+SScrollBox::Slot()
+					[
+						SAssignNew(TransactionDetailsPanel, SUndoHistoryDetails)
+						.Visibility(EVisibility::Collapsed)
+					]
+
+					+SScrollBox::Slot()
+					[
+						SAssignNew(PackageDetailsPanel, SPackageDetails)
+						.Visibility(EVisibility::Collapsed)
+					]
+				]
+
+				+SOverlay::Slot()
+				[
+					SAssignNew(NoDetailsPanel, SBox)
+					.HAlign(HAlign_Center)
+					.VAlign(VAlign_Center)
+					.Visibility(EVisibility::Visible)
+					[
+						SNew(STextBlock)
+						.Text(this, &SConcertSessionActivities::GetNoDetailsText)
+					]
+				]
+
+				+SOverlay::Slot()
+				[
+					SAssignNew(LoadingDetailsPanel, SBox)
+					.HAlign(HAlign_Center)
+					.VAlign(VAlign_Center)
+					.Visibility(EVisibility::Collapsed)
+					[
+						SNew(SThrobber)
+					]
+				]
+			]
+		]
+	];
+
+	// Check if some activities are already available.
+	FetchActivities();
+
+	if (bAutoScrollDesired)
+	{
+		FSlateApplication::Get().OnPostTick().AddSP(this, &SConcertSessionActivities::OnPostTick);
+	}
+}
+
+void SConcertSessionActivities::Tick(const FGeometry& AllottedGeometry, const double InCurrentTime, const float InDeltaTime)
+{
+	EConcertActivityFilterFlags LatestFilterFlags = QueryActiveActivityFilters();
+	if (ActiveFilterFlags != LatestFilterFlags)
+	{
+		ActiveFilterFlags = LatestFilterFlags;
+		OnActivityFilterUpdated();
+	}
+
+	FetchActivities(); // Check if we should fetch more activities in case we filtered out to many of them.
+}
+
+FText SConcertSessionActivities::GetNoDetailsText() const
+{
+	return GetSelectedActivity() ?
+		LOCTEXT("NoDetails_NotAvailable", "The selected activity doesn't have details to display.") :
+		LOCTEXT("NoDetails_NoActivitySelected", "Select an activity to view its details.");
+}
+
+void SConcertSessionActivities::OnPostTick(float)
+{
+	// NOTE: The way the list view adjust the scroll position when the component is resized has some behaviors to consider to get auto-scrolling
+	//       working consistently. When the list view shrink (allowing less item) the scroll view doesn't remain anchored at the end. Instead
+	//       the scroll position moves a little bit up and the list view doesn't consider it as scrolling because OnListViewScrolled() is not
+	//       called. The code below detects that case and maintain the scroll position at the end when required.
+
+	if (bActivityViewScrolled) // OnListViewScrolled() was invoked. The user scrolled the activity list or enlarged the view to see more items.
+	{
+		bUserScrolling = !ActivityView->GetScrollDistanceRemaining().IsNearlyZero();
+		bActivityViewScrolled = false;
+	}
+	else if (bAutoScrollDesired && !bUserScrolling && ActivityView->GetScrollDistanceRemaining().Y > 0) // See NOTE.
+	{
+		ActivityView->ScrollToBottom(); // Ensure the scroll position is maintained at the end.
+	}
+}
+
+void SConcertSessionActivities::OnListViewScrolled(double InScrollOffset)
+{
+	bActivityViewScrolled = true;
+
+	if (FetchActivitiesFn) // This widget is responsible to populate the view.
+	{
+		if (!bAllActivitiesFetched && ActivityView->GetScrollDistance().Y > 0.7) // Should fetch more?
+		{
+			DesiredActivitiesCount += ActivitiesPerRequest; // This will request another 'page' the next time FetchActivities() is called.
+		}
+	}
+}
+
+void SConcertSessionActivities::OnListViewSelectionChanged(TSharedPtr<FConcertClientSessionActivity> InActivity, ESelectInfo::Type SelectInfo)
+{
+	UpdateDetailArea(InActivity);
+}
+
+void SConcertSessionActivities::OnDetailsAreaExpansionChanged(bool bExpanded)
+{
+	bDetailsAreaExpanded = bExpanded;
+	UpdateDetailArea(bDetailsAreaExpanded ? GetSelectedActivity() : nullptr);
+}
+
+void SConcertSessionActivities::UpdateDetailArea(TSharedPtr<FConcertClientSessionActivity> InSelectedActivity)
+{
+	if (DetailsAreaVisibility != EVisibility::Visible || !bDetailsAreaExpanded)
+	{
+		return;
+	}
+	else if (!InSelectedActivity.IsValid()) // The selection was cleared?
+	{
+		SetDetailsPanelVisibility(NoDetailsPanel.Get());
+	}
+	else if (InSelectedActivity->EventPayload) // The event payload is already bundled in the activity stream?
+	{
+		if (InSelectedActivity->Activity.EventType == EConcertSyncActivityEventType::Transaction)
+		{
+			FConcertSyncTransactionEvent TransactionEvent;
+			InSelectedActivity->EventPayload->GetTypedPayload(TransactionEvent);
+			if (TransactionEvent.Transaction.ExportedObjects.Num())
+			{
+				DisplayTransactionDetails(*InSelectedActivity, TransactionEvent.Transaction);
+			}
+			else
+			{
+				SetDetailsPanelVisibility(NoDetailsPanel.Get());
+			}
+		}
+		else if (InSelectedActivity->Activity.EventType == EConcertSyncActivityEventType::Package)
+		{
+			FConcertSyncPackageEvent PackageEvent;
+			InSelectedActivity->EventPayload->GetTypedPayload(PackageEvent);
+			checkf(!PackageEvent.Package.HasPackageData(), TEXT("UI should only request the package meta data because the package data is not useful and can be very large"));
+			DisplayPackageDetails(*InSelectedActivity, PackageEvent.PackageRevision, PackageEvent.Package.Info);
+		}
+		else // Other activity types (lock/connection) don't have details panel.
+		{
+			SetDetailsPanelVisibility(NoDetailsPanel.Get());
+		}
+	}
+	else if (InSelectedActivity->Activity.EventType == EConcertSyncActivityEventType::Transaction && GetTransactionEventFn) // A function is bound to get the transaction event?
+	{
+		SetDetailsPanelVisibility(LoadingDetailsPanel.Get());
+		TWeakPtr<SConcertSessionActivities> WeakSelf = SharedThis(this);
+		GetTransactionEventFn(*InSelectedActivity).Next([WeakSelf, InSelectedActivity](const TOptional<FConcertSyncTransactionEvent>& TransactionEvent)
+		{
+			if (TSharedPtr<SConcertSessionActivities> Self = WeakSelf.Pin()) // If 'this' object hasn't been deleted.
+			{
+				if (Self->GetSelectedActivity() == InSelectedActivity) // Ensure the activity is still selected.
+				{
+					if (TransactionEvent.IsSet() && TransactionEvent->Transaction.ExportedObjects.Num())
+					{
+						Self->DisplayTransactionDetails(*InSelectedActivity, TransactionEvent.GetValue().Transaction);
+					}
+					else
+					{
+						Self->SetDetailsPanelVisibility(Self->NoDetailsPanel.Get());
+					}
+				}
+				// else -> The details panel is presenting information for another activity (or no activity).
+			}
+			// else -> The widget was deleted.
+		});
+	}
+	else if (InSelectedActivity->Activity.EventType == EConcertSyncActivityEventType::Package && GetPackageEventFn) // A function is bound to get the package event?
+	{
+		SetDetailsPanelVisibility(LoadingDetailsPanel.Get());
+		FConcertSyncPackageEventMetaData PackageEventMetaData;
+		if (GetPackageEventFn(*InSelectedActivity, PackageEventMetaData))
+		{
+			DisplayPackageDetails(*InSelectedActivity, PackageEventMetaData.PackageRevision, PackageEventMetaData.PackageInfo);
+		}
+		else
+		{
+			SetDetailsPanelVisibility(NoDetailsPanel.Get());
+		}
+	}
+	else
+	{
+		SetDetailsPanelVisibility(NoDetailsPanel.Get());
+	}
+}
+
+EConcertActivityFilterFlags SConcertSessionActivities::QueryActiveActivityFilters() const
+{
+	// The visibility attributes are externally provided. (In practice, they are controlled from the 'View Options' check boxes).
+	EConcertActivityFilterFlags ActiveFlags = EConcertActivityFilterFlags::ShowAll;
+
+	if (ConnectionActivitiesVisibility.Get() != EVisibility::Visible)
+	{
+		ActiveFlags |= EConcertActivityFilterFlags::HideConnectionActivities;
+	}
+	if (LockActivitiesVisibility.Get() != EVisibility::Visible)
+	{
+		ActiveFlags |= EConcertActivityFilterFlags::HideLockActivities;
+	}
+	if (PackageActivitiesVisibility.Get() != EVisibility::Visible)
+	{
+		ActiveFlags |= EConcertActivityFilterFlags::HidePackageActivities;
+	}
+	if (TransactionActivitiesVisibility.Get() != EVisibility::Visible)
+	{
+		ActiveFlags |= EConcertActivityFilterFlags::HideTransactionActivities;
+	}
+	if (IgnoredActivitiesVisibility.Get() != EVisibility::Visible)
+	{
+		ActiveFlags |= EConcertActivityFilterFlags::HideIgnoredActivities;
+	}
+
+	return ActiveFlags;
+}
+
+void SConcertSessionActivities::OnActivityFilterUpdated()
+{
+	// Try preserving the selected activity.
+	TSharedPtr<FConcertClientSessionActivity> SelectedActivity = GetSelectedActivity();
+
+	// Reset the list of displayed activities.
+	Activities.Reset(AllActivities.Num());
+
+	// Apply the filter.
+	for (TSharedPtr<FConcertClientSessionActivity>& Activity : AllActivities)
+	{
+		if (PassesFilters(*Activity))
+		{
+			Activities.Add(Activity);
+		}
+	}
+
+	// Restore/reset the selected activity.
+	if (SelectedActivity && Activities.Contains(SelectedActivity))
+	{
+		ActivityView->SetItemSelection(SelectedActivity, true); // Restore previous selection.
+		ActivityView->RequestScrollIntoView(SelectedActivity);
+	}
+	else if (bAutoScrollDesired && !bUserScrolling) // No activity was selected.
+	{
+		ActivityView->ScrollToBottom();
+	}
+
+	ActivityView->RequestListRefresh();
+}
+
+void SConcertSessionActivities::FetchActivities()
+{
+	if (!FetchActivitiesFn) // Not bound?
+	{
+		return; // The widget is expected to be populated/cleared externally using Append()/Reset()
+	}
+
+	bool bRefresh = false;
+
+	// If they are still activities to fetch and the user scrolled down (or our nominal amount is not reached), request more from the server.
+	if (!bAllActivitiesFetched && (Activities.Num() < DesiredActivitiesCount))
+	{
+		FText ErrorMsg;
+		int32 FetchCount = 0; // The number of activities fetched in this iteration.
+		int32 StartInsertPos = AllActivities.Num();
+
+		bAllActivitiesFetched = FetchActivitiesFn(AllActivities, FetchCount, ErrorMsg);
+		if (ErrorMsg.IsEmpty())
+		{
+			if (FetchCount) // New activities appended?
+			{
+				for (int32 Index = StartInsertPos; Index < AllActivities.Num(); ++Index) // Append the fetched activities
+				{
+					if (PassesFilters(*AllActivities[Index]))
+					{
+						Activities.Add(AllActivities[Index]);
+						bRefresh = true;
+					}
+
+					if (AllActivities[Index]->Activity.bIgnored)
+					{
+						++IgnoredActivityNum;
+					}
+				}
+			}
+		}
+		else
+		{
+			FAsyncTaskNotificationConfig NotificationConfig;
+			NotificationConfig.bIsHeadless = false;
+			NotificationConfig.bKeepOpenOnFailure = true;
+			NotificationConfig.LogCategory = &LogConcert;
+
+			FAsyncTaskNotification Notification(NotificationConfig);
+			Notification.SetComplete(LOCTEXT("FetchError", "Failed to retrieve session activities"), ErrorMsg, /*Success*/ false);
+		}
+	}
+
+	if (Activities.Num() && ActivityView->GetSelectedItems().Num() == 0)
+	{
+		ActivityView->SetItemSelection(Activities[0], true);
+	}
+
+	if (bRefresh)
+	{
+		if (bAutoScrollDesired && !bUserScrolling)
+		{
+			ActivityView->ScrollToBottom();
+		}
+
+		ActivityView->RequestListRefresh();
+	}
+}
+
+void SConcertSessionActivities::Append(TSharedPtr<FConcertClientSessionActivity> Activity)
+{
+	if (Activity->Activity.bIgnored)
+	{
+		++IgnoredActivityNum;
+	}
+
+	AllActivities.Add(Activity);
+	if (PassesFilters(*Activity))
+	{
+		Activities.Add(MoveTemp(Activity));
+
+		if (bAutoScrollDesired && !bUserScrolling)
+		{
+			ActivityView->ScrollToBottom();
+		}
+
+		ActivityView->RequestListRefresh();
+	}
+}
+
+void SConcertSessionActivities::RequestRefresh()
+{
+	ActivityView->RequestListRefresh();
+}
+
+void SConcertSessionActivities::Reset()
+{
+	Activities.Reset();
+	AllActivities.Reset();
+	ActivityView->RequestListRefresh();
+	bAllActivitiesFetched = false;
+	bUserScrolling = false;
+	DesiredActivitiesCount = ActivitiesPerRequest;
+	IgnoredActivityNum = 0;
+}
+
+bool SConcertSessionActivities::PassesFilters(const FConcertClientSessionActivity& Activity)
+{
+	if (Activity.Activity.EventType == EConcertSyncActivityEventType::Connection && ConnectionActivitiesVisibility.Get() != EVisibility::Visible) // Filter out 'connection' activities?
+	{
+		return false;
+	}
+	else if (Activity.Activity.EventType == EConcertSyncActivityEventType::Lock && LockActivitiesVisibility.Get() != EVisibility::Visible) // Filter out 'lock' activities?
+	{
+		return false;
+	}
+	else if (Activity.Activity.EventType == EConcertSyncActivityEventType::Package && PackageActivitiesVisibility.Get() != EVisibility::Visible) // Filter out 'package' activities?
+	{
+		return false;
+	}
+	else if (Activity.Activity.EventType == EConcertSyncActivityEventType::Transaction && TransactionActivitiesVisibility.Get() != EVisibility::Visible) // Filter out 'transaction' activities?
+	{
+		return false;
+	}
+	else if (Activity.Activity.bIgnored && IgnoredActivitiesVisibility.Get() != EVisibility::Visible) // Filter out 'ignored' activities?
+	{
+		return false;
+	}
+
+	return SearchTextFilter->PassesFilter(Activity);
+}
+
+FText SConcertSessionActivities::UpdateTextFilter(const FText& InFilterText)
+{
+	SearchTextFilter->SetRawFilterText(InFilterText);
+	return SearchTextFilter->GetFilterErrorText();
+}
+
+void SConcertSessionActivities::PopulateSearchStrings(const FConcertClientSessionActivity& Activity, TArray<FString>& OutSearchStrings) const
+{
+	FText ClientName = GetActivityUserFn ? ConcertSessionActivityUtils::GetClientName(GetActivityUserFn(Activity.Activity.EndpointId)) : FText::GetEmpty();
+
+	OutSearchStrings.Add(ConcertSessionActivityUtils::GetActivityDateTime(Activity, TimeFormat.Get()).ToString());
+	OutSearchStrings.Add(ConcertSessionActivityUtils::GetSummary(Activity, ClientName, /*bAsRichText*/false).ToString());
+
+	if (ClientNameColumnVisibility.Get() == EVisibility::Visible)
+	{
+		OutSearchStrings.Add(ClientName.ToString());
+	}
+
+	if (OperationColumnVisibility.Get() == EVisibility::Visible)
+	{
+		OutSearchStrings.Add(ConcertSessionActivityUtils::GetOperationName(Activity).ToString());
+	}
+
+	if (PackageColumnVisibility.Get() == EVisibility::Visible)
+	{
+		OutSearchStrings.Add(ConcertSessionActivityUtils::GetPackageName(Activity).ToString());
+	}
+}
+
+TSharedRef<ITableRow> SConcertSessionActivities::OnGenerateActivityRowWidget(TSharedPtr<FConcertClientSessionActivity> Activity, const TSharedRef<STableViewBase>& OwnerTable)
+{
+	return SNew(SConcertSessionActivityRow, Activity, GetActivityUserFn ? GetActivityUserFn(Activity->Activity.EndpointId) : nullptr, OwnerTable)
+		.TimeFormat(TimeFormat)
+		.HighlightText(HighlightText)
+		.OnMakeColumnOverlayWidget(MakeColumnOverlayWidgetFn);
+}
+
+TSharedPtr<FConcertClientSessionActivity> SConcertSessionActivities::GetSelectedActivity() const
+{
+	TArray<TSharedPtr<FConcertClientSessionActivity>> SelectedItems = ActivityView->GetSelectedItems();
+	return SelectedItems.Num() ? SelectedItems[0] : nullptr;
+}
+
+TSharedPtr<FConcertClientSessionActivity> SConcertSessionActivities::GetMostRecentActivity() const
+{
+	// NOTE: This function assumes that activities are sorted by ID. When used for recovery purpose,
+	//       the activities are listed from the most recent to the oldest. When displaying a live
+	//       session activity stream, the activities are listed from the oldest to the newest.
+
+	if (AllActivities.Num()) // Ignore the filters.
+	{
+		if (AllActivities[0]->Activity.ActivityId > AllActivities.Last()->Activity.ActivityId)
+		{
+			return AllActivities[0]; // Listed from the latest to oldest.
+		}
+		return AllActivities.Last(); // Listed from the oldest to latest.
+	}
+	return nullptr; // The list is empty.
+}
+
+bool SConcertSessionActivities::IsLastColumn(const FName& ColumnId) const
+{
+	return ColumnId == ConcertSessionActivityUtils::SummaryColumnId; // Summary column is always visible and always the last.
+}
+
+void SConcertSessionActivities::DisplayTransactionDetails(const FConcertClientSessionActivity& Activity, const FConcertTransactionEventBase& InTransaction)
+{
+	const FConcertSyncTransactionActivitySummary* Summary = Activity.ActivitySummary.Cast<FConcertSyncTransactionActivitySummary>();
+	FString TransactionTitle = Summary ? Summary->TransactionTitle.ToString() : FString();
+
+	FTransactionDiff TransactionDiff{ InTransaction.TransactionId, TransactionTitle };
+
+	for (const FConcertExportedObject& ExportedObject : InTransaction.ExportedObjects)
+	{
+		FTransactionObjectDeltaChange DeltaChange;
+		Algo::Transform(ExportedObject.PropertyDatas, DeltaChange.ChangedProperties, [](const FConcertSerializedPropertyData& PropertyData) { return PropertyData.PropertyName; });
+
+		DeltaChange.bHasNameChange = ExportedObject.ObjectData.NewName != FName();
+		DeltaChange.bHasOuterChange = ExportedObject.ObjectData.NewOuterPathName != FName();
+		DeltaChange.bHasExternalPackageChange = ExportedObject.ObjectData.NewExternalPackageName != FName();
+		DeltaChange.bHasPendingKillChange = ExportedObject.ObjectData.bIsPendingKill;
+
+		FString ObjectPathName = ExportedObject.ObjectId.ObjectOuterPathName.ToString() + TEXT(".") + ExportedObject.ObjectId.ObjectName.ToString();
+		TSharedPtr<FTransactionObjectEvent> Event = MakeShared<FTransactionObjectEvent>(InTransaction.TransactionId, InTransaction.OperationId, ETransactionObjectEventType::Finalized, MoveTemp(DeltaChange), nullptr
+			, ExportedObject.ObjectId.ObjectPackageName
+			, ExportedObject.ObjectId.ObjectName
+			, FName(*MoveTemp(ObjectPathName))
+			, ExportedObject.ObjectId.ObjectOuterPathName
+			, ExportedObject.ObjectId.ObjectExternalPackageName
+			, ExportedObject.ObjectId.ObjectClassPathName);
+
+		TransactionDiff.DiffMap.Emplace(FName(*ObjectPathName), MoveTemp(Event));
+	}
+
+	TransactionDetailsPanel->SetSelectedTransaction(MoveTemp(TransactionDiff));
+	SetDetailsPanelVisibility(TransactionDetailsPanel.Get());
+}
+
+void SConcertSessionActivities::DisplayPackageDetails(const FConcertClientSessionActivity& Activity, int64 PackageRevision, const FConcertPackageInfo& PackageInfo)
+{
+	const FConcertClientInfo* ClientInfo = nullptr;
+	if (GetActivityUserFn)
+	{
+		ClientInfo = GetActivityUserFn(Activity.Activity.EndpointId);
+	}
+
+	PackageDetailsPanel->SetPackageInfo(PackageInfo, PackageRevision, ClientInfo ? ClientInfo->DisplayName : FString());
+	SetDetailsPanelVisibility(PackageDetailsPanel.Get());
+}
+
+void SConcertSessionActivities::SetDetailsPanelVisibility(const SWidget* VisiblePanel)
+{
+	TransactionDetailsPanel->SetVisibility(VisiblePanel == TransactionDetailsPanel.Get() ? EVisibility::Visible : EVisibility::Collapsed);
+	PackageDetailsPanel->SetVisibility(VisiblePanel == PackageDetailsPanel.Get() ? EVisibility::Visible : EVisibility::Collapsed);
+	NoDetailsPanel->SetVisibility(VisiblePanel == NoDetailsPanel.Get() ? EVisibility::Visible : EVisibility::Collapsed);
+	LoadingDetailsPanel->SetVisibility(VisiblePanel == LoadingDetailsPanel.Get() ? EVisibility::Visible : EVisibility::Collapsed);
+}
+
+TSharedRef<SWidget> FConcertSessionActivitiesOptions::MakeMenuWidget()
+{
+	FMenuBuilder MenuBuilder(/*bInShouldCloseWindowAfterMenuSelection=*/true, nullptr);
+
+	MenuBuilder.AddMenuEntry(
+		LOCTEXT("DisplayRelativeTime", "Display Relative Time"),
+		LOCTEXT("DisplayRelativeTime_Tooltip", "Displays Time Relative to the Current Time"),
+		FSlateIcon(),
+		FUIAction(
+			FExecuteAction::CreateSP(this, &FConcertSessionActivitiesOptions::OnOptionToggled, ConcertSessionActivityUtils::DisplayRelativeTimeCheckBoxId),
+			FCanExecuteAction::CreateLambda([] { return true; }),
+			FIsActionChecked::CreateLambda([this] { return bDisplayRelativeTime; })),
+		NAME_None,
+		EUserInterfaceActionType::ToggleButton
+	);
+
+	if (bEnablePackageActivityFiltering)
+	{
+		MenuBuilder.AddMenuEntry(
+			LOCTEXT("ShowPackageActivities", "Show Package Activities"),
+			LOCTEXT("ShowPackageActivities_Tooltip", "Displays create/save/rename/delete package events."),
+			FSlateIcon(),
+			FUIAction(
+				FExecuteAction::CreateSP(this, &FConcertSessionActivitiesOptions::OnOptionToggled, ConcertSessionActivityUtils::ShowPackageActivitiesCheckBoxId),
+				FCanExecuteAction::CreateLambda([] { return true; }),
+				FIsActionChecked::CreateLambda([this] { return bDisplayPackageActivities; })),
+			NAME_None,
+			EUserInterfaceActionType::ToggleButton
+		);
+	}
+
+	if (bEnableTransactionActivityFiltering)
+	{
+		MenuBuilder.AddMenuEntry(
+			LOCTEXT("ShowTransactionActivities", "Show Transaction Activities"),
+			LOCTEXT("ShowTransactionActivities_Tooltip", "Displays changes performed on assets."),
+			FSlateIcon(),
+			FUIAction(
+				FExecuteAction::CreateSP(this, &FConcertSessionActivitiesOptions::OnOptionToggled, ConcertSessionActivityUtils::ShowTransactionActivitiesCheckBoxId),
+				FCanExecuteAction::CreateLambda([] { return true; }),
+				FIsActionChecked::CreateLambda([this] { return bDisplayTransactionActivities; })),
+			NAME_None,
+			EUserInterfaceActionType::ToggleButton
+		);
+	}
+
+	if (bEnableConnectionActivityFiltering)
+	{
+		MenuBuilder.AddMenuEntry(
+			LOCTEXT("ShowConnectionActivities", "Show Connection Activities"),
+			LOCTEXT("ShowConnectionActivities_Tooltip", "Displays when client joined or left the session"),
+			FSlateIcon(),
+			FUIAction(
+				FExecuteAction::CreateSP(this, &FConcertSessionActivitiesOptions::OnOptionToggled, ConcertSessionActivityUtils::ShowConnectionActivitiesCheckBoxId),
+				FCanExecuteAction::CreateLambda([] { return true; }),
+				FIsActionChecked::CreateLambda([this] { return bDisplayConnectionActivities; })),
+			NAME_None,
+			EUserInterfaceActionType::ToggleButton
+		);
+	}
+
+	if (bEnableLockActivityFiltering)
+	{
+		MenuBuilder.AddMenuEntry(
+			LOCTEXT("ShowLockActivities", "Show Lock Activities"),
+			LOCTEXT("ShowLockActivities_Tooltip", "Displays lock/unlock events"),
+			FSlateIcon(),
+			FUIAction(
+				FExecuteAction::CreateSP(this, &FConcertSessionActivitiesOptions::OnOptionToggled, ConcertSessionActivityUtils::ShowLockActivitiesCheckBoxId),
+				FCanExecuteAction::CreateLambda([] { return true; }),
+				FIsActionChecked::CreateLambda([this] { return bDisplayLockActivities; })),
+			NAME_None,
+			EUserInterfaceActionType::ToggleButton
+		);
+	}
+
+	if (bEnableIgnoredActivityFiltering)
+	{
+		MenuBuilder.AddMenuEntry(
+			LOCTEXT("ShowIgnoredActivities", "Show Unrecoverable Activities"),
+			LOCTEXT("ShowIgnoredActivities_Tooltip", "Displays activities that were recorded, but could not be recovered in this context."),
+			FSlateIcon(),
+			FUIAction(
+				FExecuteAction::CreateSP(this, &FConcertSessionActivitiesOptions::OnOptionToggled, ConcertSessionActivityUtils::ShowIgnoredActivitiesCheckBoxId),
+				FCanExecuteAction::CreateLambda([] { return true; }),
+				FIsActionChecked::CreateLambda([this] { return bDisplayIgnoredActivities; })),
+			NAME_None,
+			EUserInterfaceActionType::ToggleButton
+		);
+	}
+	
+	return MenuBuilder.MakeWidget();
+}
+
+TSharedRef<SWidget> FConcertSessionActivitiesOptions::MakeViewOptionsWidget()
+{
+	return SNew(SComboButton)
+	.ComboButtonStyle(FAppStyle::Get(), "GenericFilters.ComboButtonStyle")
+	.ForegroundColor(FLinearColor::White)
+	.ContentPadding(0)
+	.OnGetMenuContent(this, &FConcertSessionActivitiesOptions::MakeMenuWidget)
+	.HasDownArrow(true)
+	.ContentPadding(FMargin(1, 0))
+	.ButtonContent()
+	[
+		SNew(SHorizontalBox)
+
+		+SHorizontalBox::Slot()
+		.AutoWidth()
+		.VAlign(VAlign_Center)
+		[
+			SNew(SImage).Image(FAppStyle::Get().GetBrush("GenericViewButton")) // The eye ball image.
+		]
+
+		+SHorizontalBox::Slot()
+		.AutoWidth()
+		.Padding(2, 0, 0, 0)
+		.VAlign(VAlign_Center)
+		[
+			SNew(STextBlock).Text(LOCTEXT("ViewOptions", "View Options"))
+		]
+	];
+}
+
+TSharedRef<SWidget> FConcertSessionActivitiesOptions::MakeDisplayedActivityCountWidget(TAttribute<int32> Total, TAttribute<int32> Displayed)
+{
+	return SNew(STextBlock)
+		.Text_Lambda([Total = MoveTemp(Total), Displayed = MoveTemp(Displayed)]()
+		{
+			if (Total.Get() == Displayed.Get())
+			{
+				return FText::Format(LOCTEXT("OperationCount", "{0} operations"), Total.Get());
+			}
+			else
+			{
+				return FText::Format(LOCTEXT("PartialOperationCount", "Showing {0} of {1} {1}|plural(one=operation,other=operations)"), Displayed.Get(), Total.Get());
+			}
+		});
+}
+
+TSharedRef<SWidget> FConcertSessionActivitiesOptions::MakeStatusBar(TAttribute<int32> Total, TAttribute<int32> Displayed)
+{
+	return SNew(SHorizontalBox)
+
+	// Operation count.
+	+SHorizontalBox::Slot()
+	.AutoWidth()
+	.VAlign(VAlign_Center)
+	[
+		MakeDisplayedActivityCountWidget(MoveTemp(Total), MoveTemp(Displayed))
+	]
+
+	// Gap Filler
+	+SHorizontalBox::Slot()
+	.FillWidth(1.0)
+	[
+		SNew(SSpacer)
+	]
+
+	// View option
+	+SHorizontalBox::Slot()
+	.AutoWidth()
+	[
+		MakeViewOptionsWidget()
+	];
+}
+
+void FConcertSessionActivitiesOptions::OnOptionToggled(const FName CheckBoxId)
+{
+	if (CheckBoxId == ConcertSessionActivityUtils::DisplayRelativeTimeCheckBoxId)
+	{
+		bDisplayRelativeTime = !bDisplayRelativeTime;
+	}
+	else if (CheckBoxId == ConcertSessionActivityUtils::ShowConnectionActivitiesCheckBoxId)
+	{
+		bDisplayConnectionActivities = !bDisplayConnectionActivities;
+	}
+	else if (CheckBoxId == ConcertSessionActivityUtils::ShowLockActivitiesCheckBoxId)
+	{
+		bDisplayLockActivities = !bDisplayLockActivities;
+	}
+	else if (CheckBoxId == ConcertSessionActivityUtils::ShowPackageActivitiesCheckBoxId)
+	{
+		bDisplayPackageActivities = !bDisplayPackageActivities;
+	}
+	else if (CheckBoxId == ConcertSessionActivityUtils::ShowTransactionActivitiesCheckBoxId)
+	{
+		bDisplayTransactionActivities = !bDisplayTransactionActivities;
+	}
+	else if (CheckBoxId == ConcertSessionActivityUtils::ShowIgnoredActivitiesCheckBoxId)
+	{
+		bDisplayIgnoredActivities = !bDisplayIgnoredActivities;
+	}	
+}
+
+#undef LOCTEXT_NAMESPACE
diff --git a/Engine/Plugins/Developer/Concert/ConcertUI/ConcertSharedSlate/Source/ConcertSharedSlate/Private/SConcertSessionRecovery.cpp b/Engine/Plugins/Developer/Concert/ConcertUI/ConcertSharedSlate/Source/ConcertSharedSlate/Private/SConcertSessionRecovery.cpp
new file mode 100644
--- /dev/null
+++ b/Engine/Plugins/Developer/Concert/ConcertUI/ConcertSharedSlate/Source/ConcertSharedSlate/Private/SConcertSessionRecovery.cpp
@@ -0,0 +1,322 @@
+// Copyright Epic Games, Inc. All Rights Reserved.
+
+#include "SConcertSessionRecovery.h"
+
+#include "IConcertClientWorkspace.h"
+#include "EditorStyleSet.h"
+#include "EditorFontGlyphs.h"
+#include "Widgets/SWindow.h"
+#include "Widgets/Input/SButton.h"
+#include "Widgets/Input/SSearchBox.h"
+#include "Widgets/Layout/SBorder.h"
+#include "Widgets/Layout/SBox.h"
+#include "Widgets/Layout/SSpacer.h"
+#include "Widgets/Layout/SSeparator.h"
+#include "Widgets/Layout/SUniformGridPanel.h"
+#include "SConcertSessionActivities.h"
+
+#define LOCTEXT_NAMESPACE "SConcertSessionRecovery"
+
+void SConcertSessionRecovery::Construct(const FArguments& InArgs)
+{
+	IntroductionText = InArgs._IntroductionText;
+	ParentWindow = InArgs._ParentWindow;
+	OnRestoreFn = InArgs._OnRestore;
+	OnCancelFn = InArgs._OnCancel;
+	IsRecoverThroughButtonVisible = InArgs._IsRecoverThroughButtonsVisible;
+
+	ActivityViewOptions = MakeShared<FConcertSessionActivitiesOptions>();
+	ActivityViewOptions->bEnableConnectionActivityFiltering = InArgs._IsConnectionActivityFilteringEnabled;
+	ActivityViewOptions->bEnableLockActivityFiltering = InArgs._IsLockActivityFilteringEnabled;
+	ActivityViewOptions->bEnablePackageActivityFiltering = InArgs._IsPackageActivityFilteringEnabled;
+	ActivityViewOptions->bEnableTransactionActivityFiltering = InArgs._IsTransactionActivityFilteringEnabled;
+	ActivityViewOptions->bEnableIgnoredActivityFiltering = InArgs._IsIgnoredActivityFilteringEnabled;
+
+	SAssignNew(ActivityView, SConcertSessionActivities)
+	.OnFetchActivities(InArgs._OnFetchActivities)
+	.OnMapActivityToClient(InArgs._OnMapActivityToClient)
+	.OnMakeColumnOverlayWidget([this](TWeakPtr<FConcertClientSessionActivity> Activity, const FName& ColumnId) { return MakeRecoverThroughWidget(Activity, ColumnId); })
+	.HighlightText(this, &SConcertSessionRecovery::HighlightSearchText)
+	.TimeFormat(ActivityViewOptions.Get(), &FConcertSessionActivitiesOptions::GetTimeFormat)
+	.ClientAvatarColorColumnVisibility(InArgs._ClientAvatarColorColumnVisibility)
+	.ClientNameColumnVisibility(InArgs._ClientNameColumnVisibility)
+	.OperationColumnVisibility(InArgs._OperationColumnVisibility)
+	.PackageColumnVisibility(InArgs._PackageColumnVisibility)
+	.ConnectionActivitiesVisibility(ActivityViewOptions.Get(), &FConcertSessionActivitiesOptions::GetConnectionActivitiesVisibility)
+	.LockActivitiesVisibility(ActivityViewOptions.Get(), &FConcertSessionActivitiesOptions::GetLockActivitiesVisibility)
+	.PackageActivitiesVisibility(ActivityViewOptions.Get(), &FConcertSessionActivitiesOptions::GetPackageActivitiesVisibility)
+	.TransactionActivitiesVisibility(ActivityViewOptions.Get(), &FConcertSessionActivitiesOptions::GetTransactionActivitiesVisibility)
+	.IgnoredActivitiesVisibility(ActivityViewOptions.Get(), &FConcertSessionActivitiesOptions::GetIgnoredActivitiesVisibility)
+	.DetailsAreaVisibility(InArgs._DetailsAreaVisibility)
+	.NoActivitiesReasonText(InArgs._NoActivitiesReasonText);
+
+	ChildSlot
+	[
+		SNew(SBorder)
+		.BorderImage(FEditorStyle::GetBrush("ToolPanel.GroupBorder"))
+		.BorderBackgroundColor(FSlateColor(FLinearColor(0.6, 0.6, 0.6)))
+		.Padding(0)
+		[
+			SNew(SVerticalBox)
+			+SVerticalBox::Slot()
+			.FillHeight(1.0)
+			[
+				SNew(SBorder)
+				.BorderImage(FEditorStyle::GetBrush("ToolPanel.GroupBorder"))
+				[
+					SNew(SVerticalBox)
+
+					// Contextual introduction.
+					+SVerticalBox::Slot()
+					.AutoHeight()
+					.HAlign(HAlign_Center)
+					.Padding(0, 6)
+					[
+						SNew(STextBlock)
+						.Text(IntroductionText)
+						.Visibility(IntroductionText.IsEmpty() ? EVisibility::Collapsed : EVisibility::Visible)
+					]
+
+					// Search bar.
+					+SVerticalBox::Slot()
+					.AutoHeight()
+					.Padding(0, 1, 0, 2)
+					[
+						SAssignNew(SearchBox, SSearchBox)
+						.HintText(LOCTEXT("SearchHint", "Search..."))
+						.OnTextChanged(this, &SConcertSessionRecovery::OnSearchTextChanged)
+						.OnTextCommitted(this, &SConcertSessionRecovery::OnSearchTextCommitted)
+						.DelayChangeNotificationsWhileTyping(true)
+					]
+
+					// Activity List
+					+SVerticalBox::Slot()
+					[
+						ActivityView.ToSharedRef()
+					]
+
+					+SVerticalBox::Slot()
+					.AutoHeight()
+					.Padding(2, 2)
+					[
+						SNew(SSeparator)
+					]
+
+					// Status bar/View options.
+					+SVerticalBox::Slot()
+					.AutoHeight()
+					.Padding(4, 2)
+					[
+						ActivityViewOptions->MakeStatusBar(
+							TAttribute<int32>(ActivityView.Get(), &SConcertSessionActivities::GetTotalActivityNum),
+							TAttribute<int32>(ActivityView.Get(), &SConcertSessionActivities::GetDisplayedActivityNum))
+					]
+				]
+			]
+
+			// Buttons
+			+SVerticalBox::Slot()
+			.AutoHeight()
+			.HAlign(HAlign_Right)
+			.Padding(0, 6)
+			[
+				SNew(SUniformGridPanel)
+				.SlotPadding(FMargin(2.0f, 0.0f))
+				.Visibility(InArgs._AreRecoverAllAndCancelButtonsVisible.Get() ? EVisibility::Visible : EVisibility::Collapsed)
+
+				+SUniformGridPanel::Slot(0, 0)
+				[
+					SNew(SButton)
+					.ForegroundColor(FLinearColor::White)
+					.ButtonStyle(FEditorStyle::Get(), TEXT("FlatButton.Success"))
+					.ToolTipText(this, &SConcertSessionRecovery::GetRecoverAllButtonTooltip)
+					.OnClicked(this, &SConcertSessionRecovery::OnRecoverAllClicked)
+					.HAlign(HAlign_Center)
+					.ContentPadding(FMargin(14, 3))
+					[
+						SNew(STextBlock)
+						.Text(LOCTEXT("RecoverAll", "Recover All"))
+						.Font( FEditorStyle::GetFontStyle("BoldFont"))
+						.ShadowOffset( FVector2D( 1.0f, 1.0f ) )
+					]
+				]
+
+				+SUniformGridPanel::Slot(1, 0)
+				[
+					SNew(SButton)
+					.ForegroundColor(FLinearColor::White)
+					.ButtonStyle(FEditorStyle::Get(), TEXT("FlatButton.Danger"))
+					.ToolTipText(LOCTEXT("CancelRecoveryTooltip", "Discard any recoverable data for your assets and continue with their last saved state"))
+					.OnClicked(this, &SConcertSessionRecovery::OnCancelRecoveryClicked)
+					.HAlign(HAlign_Center)
+					.ContentPadding(FMargin(14, 3))
+					[
+						SNew(STextBlock)
+						.Text(LOCTEXT("Cancel", "Cancel"))
+						.Font( FEditorStyle::GetFontStyle("BoldFont"))
+						.ShadowOffset( FVector2D( 1.0f, 1.0f ) )
+					]
+				]
+			]
+		]
+	];
+}
+
+TSharedPtr<SWidget> SConcertSessionRecovery::MakeRecoverThroughWidget(TWeakPtr<FConcertClientSessionActivity> Activity, const FName& ColumnId)
+{
+	if (TSharedPtr<FConcertClientSessionActivity> ActivityPin = Activity.Pin())
+	{
+		if (ActivityView->IsLastColumn(ColumnId) && !ActivityPin->Activity.bIgnored) // The most right cell and activity will not be ignored on recovery.
+		{
+			// The green 'Recover Through' button that appears in the most right cell if the row is selected.
+			return SNew(SBox)
+			.Padding(FMargin(1, 1))
+			.HAlign(HAlign_Right)
+			.VAlign(VAlign_Center)
+			[
+				SNew(SButton)
+				.ForegroundColor(FLinearColor::White)
+				.ButtonStyle(FEditorStyle::Get(), TEXT("FlatButton.Success"))
+				.Visibility_Lambda([this, Activity](){ return GetRecoverThroughButtonVisibility(Activity.Pin()); })
+				.OnClicked_Lambda([this, Activity](){ RecoverThrough(Activity.Pin()); return FReply::Handled(); })
+				.ToolTipText(this, &SConcertSessionRecovery::GetRecoverThroughButtonTooltip)
+				.ContentPadding(FMargin(20, 0))
+				[
+					SNew(STextBlock)
+					.Font(FEditorStyle::Get().GetFontStyle("FontAwesome.10"))
+					.Text(FEditorFontGlyphs::Arrow_Circle_O_Right)
+				]
+			];
+		}
+	}
+
+	return nullptr; // No overlay.
+}
+
+void SConcertSessionRecovery::Reset()
+{
+	ActivityView->Reset();
+}
+
+int32 SConcertSessionRecovery::GetTotalActivityNum() const
+{
+	return ActivityView->GetTotalActivityNum();
+}
+
+TSharedPtr<FConcertClientSessionActivity> SConcertSessionRecovery::GetMostRecentActivity() const
+{
+	return ActivityView->GetMostRecentActivity();
+}
+
+EVisibility SConcertSessionRecovery::GetRecoverThroughButtonVisibility(TSharedPtr<FConcertClientSessionActivity> Activity)
+{
+	return Activity == ActivityView->GetSelectedActivity() && IsRecoverThroughButtonVisible.Get() ? EVisibility::Visible : EVisibility::Hidden;
+}
+
+FText SConcertSessionRecovery::GetRecoverThroughButtonTooltip() const
+{
+	if ((ActivityView->GetTotalActivityNum() == ActivityView->GetDisplayedActivityNum() && ActivityView->GetIgnoredActivityNum() == 0) ||
+	    (ActivityView->GetTotalActivityNum() - ActivityView->GetDisplayedActivityNum() == ActivityView->GetIgnoredActivityNum()))
+	{
+		return LOCTEXT("RecoverThroughAll_Tooltip", "Recover all prior activities through this activity.");
+	}
+	else if (ActivityView->GetTotalActivityNum() == ActivityView->GetDisplayedActivityNum())
+	{
+		return LOCTEXT("RecoverThroughWithIgnored_Tooltip", "Recover all prior activities through this activity, excluding unrecoverable ones (grayed out).");
+	}
+	else if (ActivityView->GetIgnoredActivityNum() == 0 || !ActivityViewOptions->bDisplayIgnoredActivities)
+	{
+		return LOCTEXT("RecoverThroughWithFilter_Tooltip", "Recover all prior activities through this activity, including the ones currently filtered out by a text search or a view option.");
+	}
+
+	return LOCTEXT("RecoverThroughWithIgnoredAndFilter_Tooltip", "Recover all prior activities through this activity, excluding unrecoverable ones (grayed out), but including the those currently filtered out by a text search or a view option.");
+}
+
+FText SConcertSessionRecovery::GetRecoverAllButtonTooltip() const
+{
+	if ((ActivityView->GetTotalActivityNum() == ActivityView->GetDisplayedActivityNum() && ActivityView->GetIgnoredActivityNum() == 0) ||
+	    (ActivityView->GetTotalActivityNum() - ActivityView->GetDisplayedActivityNum() == ActivityView->GetIgnoredActivityNum()))
+	{
+		return LOCTEXT("RecoverAll_Tooltip", "Recover all activities.");
+	}
+	else if (ActivityView->GetTotalActivityNum() == ActivityView->GetDisplayedActivityNum())
+	{
+		return LOCTEXT("RecoverAllWithIgnored_Tooltip", "Recover all activities, excluding unrecoverable ones (grayed out).");
+	}
+	else if (ActivityView->GetIgnoredActivityNum() == 0 || !ActivityViewOptions->bDisplayIgnoredActivities)
+	{
+		return LOCTEXT("RecoverAllWithFilter_Tooltip", "Recover all activities, including the ones currently filtered out by a text search or a view option.");
+	}
+
+	return LOCTEXT("RecoverAllWithIgnoredAndFilter_Tooltip", "Recover all activities, excluding unrecoverable ones (grayed out), but including those currently filtered out by a text search or a view option.");
+}
+
+void SConcertSessionRecovery::OnSearchTextChanged(const FText& InFilterText)
+{
+	SearchText = InFilterText;
+	SearchBox->SetError(ActivityView->UpdateTextFilter(InFilterText));
+}
+
+void SConcertSessionRecovery::OnSearchTextCommitted(const FText& InFilterText, ETextCommit::Type CommitType)
+{
+	if (!InFilterText.EqualTo(SearchText))
+	{
+		OnSearchTextChanged(InFilterText);
+	}
+}
+
+FText SConcertSessionRecovery::HighlightSearchText() const
+{
+	return SearchText;
+}
+
+FReply SConcertSessionRecovery::OnCancelRecoveryClicked()
+{
+	check(!RecoveryThroughItem.IsValid());
+
+	if (OnCancelFn) // If a cancel behavior is bound.
+	{
+		OnCancelFn();
+	}
+	else // Default cancel behavior.
+	{
+		DismissWindow();
+	}
+	return FReply::Handled();
+}
+
+FReply SConcertSessionRecovery::OnRecoverAllClicked()
+{
+	// Recover to the most recent activity, ignoring any filter being applied to the view.
+	RecoverThrough(ActivityView->GetMostRecentActivity());
+	return FReply::Handled();
+}
+
+void SConcertSessionRecovery::RecoverThrough(TSharedPtr<FConcertClientSessionActivity> Item)
+{
+	bool bShouldDismissWindow = true;
+	if (Item)
+	{
+		RecoveryThroughItem = Item;
+
+		if (OnRestoreFn)
+		{
+			bShouldDismissWindow = OnRestoreFn(RecoveryThroughItem);
+		}
+	}
+
+	if (bShouldDismissWindow)
+	{
+		DismissWindow();
+	}
+}
+
+void SConcertSessionRecovery::DismissWindow()
+{
+	if (TSharedPtr<SWindow> ParentWindowPin = ParentWindow.Pin())
+	{
+		ParentWindowPin->RequestDestroyWindow();
+	}
+}
+
+#undef LOCTEXT_NAMESPACE
diff --git a/Engine/Plugins/Developer/Concert/ConcertUI/ConcertSharedSlate/Source/ConcertSharedSlate/Private/SPackageDetails.cpp b/Engine/Plugins/Developer/Concert/ConcertUI/ConcertSharedSlate/Source/ConcertSharedSlate/Private/SPackageDetails.cpp
new file mode 100644
--- /dev/null
+++ b/Engine/Plugins/Developer/Concert/ConcertUI/ConcertSharedSlate/Source/ConcertSharedSlate/Private/SPackageDetails.cpp
@@ -0,0 +1,168 @@
+// Copyright Epic Games, Inc. All Rights Reserved.
+
+#include "SPackageDetails.h"
+
+#include "ConcertWorkspaceData.h"
+#include "Widgets/Layout/SBox.h"
+#include "SlateOptMacros.h"
+
+#define LOCTEXT_NAMESPACE "SPackageDetails"
+
+namespace PackageDetailsUI
+{
+	static const FName RevisionLabel(TEXT("Revision"));
+	static const FName NameLabel(TEXT("Name"));
+	static const FName ModifiedByLabel(TEXT("ModifiedBy"));
+	static const FName ActionLabel(TEXT("Action"));
+}
+
+class SPackageDetailsRow : public SMultiColumnTableRow<TSharedPtr<int32>>
+{
+	SLATE_BEGIN_ARGS(SPackageDetailsRow) {}
+		SLATE_ARGUMENT(FText, Revision)
+		SLATE_ARGUMENT(FText, Name)
+		SLATE_ARGUMENT(FText, ModifiedBy)
+		SLATE_ARGUMENT(FText, Action)
+	SLATE_END_ARGS()
+
+public:
+
+	/**
+	 * Constructs the widget.
+	 *
+	 * @param InArgs The construction arguments.
+	 */
+	void Construct(const FArguments& InArgs, const TSharedRef<STableViewBase>& InOwnerTableView)
+	{
+		Revision = InArgs._Revision;
+		Name = InArgs._Name;
+		ModifiedBy = InArgs._ModifiedBy;
+		Action = InArgs._Action;
+
+		SMultiColumnTableRow<TSharedPtr<int32> >::Construct(FSuperRowType::FArguments(), InOwnerTableView);
+	}
+
+public:
+
+	BEGIN_SLATE_FUNCTION_BUILD_OPTIMIZATION
+	virtual TSharedRef<SWidget> GenerateWidgetForColumn(const FName& ColumnName) override
+	{
+		FText ColumnText;
+
+		if (ColumnName == PackageDetailsUI::RevisionLabel)
+		{
+			ColumnText = Revision;
+		}
+		else if (ColumnName == PackageDetailsUI::NameLabel)
+		{
+			ColumnText = Name;
+		}
+		else if (ColumnName == PackageDetailsUI::ModifiedByLabel)
+		{
+			ColumnText = ModifiedBy;
+		}
+		else if (ColumnName == PackageDetailsUI::ActionLabel)
+		{
+			ColumnText = Action;
+		}
+
+		return SNew(SBox)
+			.Padding(FMargin(4.0, 0.0))
+			[
+				SNew(STextBlock)
+				.Text(ColumnText)
+			];
+	}
+	END_SLATE_FUNCTION_BUILD_OPTIMIZATION
+
+private:
+
+	/** Package revision. */
+	FText Revision;
+
+	/** Package name. */
+	FText Name;
+
+	/** Name of the user who modified the package. */
+	FText ModifiedBy;
+
+	/** What the modification was. */
+	FText Action;
+};
+
+void SPackageDetails::Construct(const FArguments& InArgs)
+{
+	ChildSlot
+		[
+			SAssignNew(DetailListView, SListView<FListItemPtr>)
+			.ListItemsSource(&DetailItems)
+			.OnGenerateRow(this, &SPackageDetails::HandleGenerateRow)
+			.HeaderRow
+			(
+				SNew(SHeaderRow)
+				+SHeaderRow::Column(PackageDetailsUI::RevisionLabel)
+				.DefaultLabel(LOCTEXT("RevisionFlagsHeaderName", "Revision"))
+				+SHeaderRow::Column(PackageDetailsUI::NameLabel)
+				.DefaultLabel(LOCTEXT("NameColumnHeaderName", "Name"))
+				+SHeaderRow::Column(PackageDetailsUI::ModifiedByLabel)
+				.DefaultLabel(LOCTEXT("ModifiedByColumnHeaderName", "ModifiedBy"))
+				+SHeaderRow::Column(PackageDetailsUI::ActionLabel)
+				.DefaultLabel(LOCTEXT("ActionColumnHeaderName", "Action"))
+			)
+		];
+}
+
+void SPackageDetails::SetPackageInfo(const FConcertPackageInfo& InPackageInfo, const int64 InRevision, const FString& InModifiedBy)
+{
+	DetailItems.Reset();
+	
+	FText Action = LOCTEXT("InvalidPackageAction", "Invalid Action");
+
+	switch (InPackageInfo.PackageUpdateType)
+	{
+		case EConcertPackageUpdateType::Added:
+			Action = LOCTEXT("PackageActionAdd", "Add");
+			break;
+		case EConcertPackageUpdateType::Deleted:
+			Action = LOCTEXT("PackageActionDelete", "Delete");
+			break;
+		case EConcertPackageUpdateType::Renamed:
+			Action = LOCTEXT("PackageActionRename", "Rename/Move");
+			break;
+		case EConcertPackageUpdateType::Saved:
+			Action = LOCTEXT("PackageActionSave", "Save");
+			break;
+		default:
+			break;
+	}
+
+	FPackageDetailsRow Row
+	{
+		FText::FromName(InPackageInfo.PackageName),
+		FText::FromString(LexToString(InRevision)),
+		InModifiedBy.IsEmpty() ? LOCTEXT("UnknownUser", "Unknown") : FText::FromString(InModifiedBy),
+		MoveTemp(Action)
+	};
+
+	DetailItems.Emplace(MakeShared<FPackageDetailsRow>(MoveTemp(Row)));
+
+	if (InPackageInfo.NewPackageName != FName())
+	{
+		FPackageDetailsRow RenamedRow;
+		RenamedRow.PackageName = FText::FromString(TEXT("-> ") + InPackageInfo.NewPackageName.ToString());
+		DetailItems.Emplace(MakeShared<FPackageDetailsRow>(MoveTemp(RenamedRow)));
+	}
+
+	DetailListView->RequestListRefresh();
+}
+
+TSharedRef<ITableRow> SPackageDetails::HandleGenerateRow(FListItemPtr InItem, const TSharedRef<STableViewBase>& OwnerTable) const
+{	
+	return SNew(SPackageDetailsRow, OwnerTable)
+		.Revision(InItem->Revision)
+		.Name(InItem->PackageName)
+		.ModifiedBy(InItem->ModifiedBy)
+		.Action(InItem->Action);
+}
+
+#undef LOCTEXT_NAMESPACE /** SPackageDetails */
diff --git a/Engine/Plugins/Developer/Concert/ConcertUI/ConcertSharedSlate/Source/ConcertSharedSlate/Public/ConcertFrontendStyle.h b/Engine/Plugins/Developer/Concert/ConcertUI/ConcertSharedSlate/Source/ConcertSharedSlate/Public/ConcertFrontendStyle.h
new file mode 100644
--- /dev/null
+++ b/Engine/Plugins/Developer/Concert/ConcertUI/ConcertSharedSlate/Source/ConcertSharedSlate/Public/ConcertFrontendStyle.h
@@ -0,0 +1,24 @@
+// Copyright Epic Games, Inc. All Rights Reserved.
+
+#pragma once
+
+#include "CoreMinimal.h"
+#include "Styling/ISlateStyle.h"
+#include "Styling/SlateStyle.h"
+
+class CONCERTSHAREDSLATE_API FConcertFrontendStyle
+{
+public:
+	
+	static void Initialize();
+	static void Shutdown();
+
+	static TSharedPtr< class ISlateStyle > Get();
+
+	static FName GetStyleSetName();
+	
+private:
+	
+	static FString InContent(const FString& RelativePath, const ANSICHAR* Extension);
+	static TSharedPtr< class FSlateStyleSet > StyleSet;
+};
diff --git a/Engine/Plugins/Developer/Concert/ConcertUI/ConcertSharedSlate/Source/ConcertSharedSlate/Public/ConcertFrontendUtils.h b/Engine/Plugins/Developer/Concert/ConcertUI/ConcertSharedSlate/Source/ConcertSharedSlate/Public/ConcertFrontendUtils.h
new file mode 100644
--- /dev/null
+++ b/Engine/Plugins/Developer/Concert/ConcertUI/ConcertSharedSlate/Source/ConcertSharedSlate/Public/ConcertFrontendUtils.h
@@ -0,0 +1,215 @@
+// Copyright Epic Games, Inc. All Rights Reserved.
+
+#pragma once
+
+#include "CoreMinimal.h"
+#include "ConcertActionDefinition.h"
+#include "Styling/AppStyle.h"
+#include "Widgets/Layout/SBorder.h"
+#include "Widgets/Images/SImage.h"
+#include "Widgets/Input/SButton.h"
+#include "Widgets/Text/STextBlock.h"
+#include "Widgets/Layout/SBox.h"
+#include "Widgets/Layout/SExpandableArea.h"
+#include "Widgets/SBoxPanel.h"
+#include "EditorFontGlyphs.h"
+#include "ConcertFrontendStyle.h"
+
+#define LOCTEXT_NAMESPACE "ConcertFrontendUtils"
+
+namespace ConcertFrontendUtils
+{
+	static const FName ButtonIconSyle = TEXT("FontAwesome.10");
+	static const float MinDesiredWidthForBtnAndIcon = 29.f;
+	static const FName ButtonStyleNames[(int32)EConcertActionType::NUM] = {
+		TEXT("SimpleButton"),
+		TEXT("FlatButton.Primary"),
+		TEXT("FlatButton.Info"),
+		TEXT("FlatButton.Success"),
+		TEXT("FlatButton.Warning"),
+		TEXT("FlatButton.Danger"),
+	};
+
+	FORCEINLINE bool ShowSessionConnectionUI()
+	{
+		return !(IS_PROGRAM);
+	}
+
+	inline TSharedRef<SWidget> CreateDisplayName(const TAttribute<FText>& InDisplayName)
+	{
+		return SNew(SBorder)
+			.BorderImage(FAppStyle::Get().GetBrush("NoBorder"))
+			.ColorAndOpacity(FLinearColor(0.75f, 0.75f, 0.75f))
+			.Padding(FMargin(6.0f, 4.0f))
+			[
+				SNew(STextBlock)
+				.Font(FAppStyle::Get().GetFontStyle("BoldFont"))
+				.Text(InDisplayName)
+			];
+	}
+
+	inline TSharedRef<SButton> CreateTextButton(const FConcertActionDefinition& InDef)
+	{
+		const FButtonStyle* ButtonStyle = &FAppStyle::Get().GetWidgetStyle<FButtonStyle>(ButtonStyleNames[(int32)InDef.Type]);
+		check(ButtonStyle);
+		const float ButtonContentWidthPadding = 6.f;
+		const float PaddingCompensation = (ButtonStyle->NormalPadding.Left + ButtonStyle->NormalPadding.Right + ButtonContentWidthPadding * 2);
+
+		return SNew(SButton)
+			.ButtonStyle(ButtonStyle)
+			.ToolTipText(InDef.ToolTipText)
+			.ContentPadding(FMargin(ButtonContentWidthPadding, 2.f))
+			.IsEnabled(InDef.IsEnabled)
+			.Visibility_Lambda([IsVisible = InDef.IsVisible]() { return IsVisible.Get() ? EVisibility::Visible : EVisibility::Collapsed; })
+			.OnClicked_Lambda([OnExecute = InDef.OnExecute]() { OnExecute.ExecuteIfBound(); return FReply::Handled(); })
+			[
+				SNew(SBox)
+				.MinDesiredWidth(MinDesiredWidthForBtnAndIcon - PaddingCompensation)
+				.HAlign(HAlign_Center)
+				[
+					SNew(STextBlock)
+					.Font(FAppStyle::Get().GetFontStyle(ButtonIconSyle))
+					.Text(InDef.Text)
+					.Justification(ETextJustify::Center)
+				]
+			];
+	}
+
+	inline TSharedRef<SButton> CreateIconButton(const FConcertActionDefinition& InDef)
+	{
+		return SNew(SButton)
+			.ButtonStyle(FAppStyle::Get(), ButtonStyleNames[(int32)InDef.Type])
+			.ToolTipText(InDef.ToolTipText)
+			.ContentPadding(FMargin(0, 0))
+			.IsEnabled(InDef.IsEnabled)
+			.Visibility_Lambda([IsVisible = InDef.IsVisible]() { return IsVisible.Get() ? EVisibility::Visible : EVisibility::Collapsed; })
+			.OnClicked_Lambda([OnExecute = InDef.OnExecute]() { OnExecute.ExecuteIfBound(); return FReply::Handled(); })
+			.VAlign(VAlign_Center)
+			.HAlign(HAlign_Center)
+			[
+				SNew(SImage)
+				.Image(TAttribute<const FSlateBrush*>::Create([IconStyleAttr = InDef.IconStyle]() { return FConcertFrontendStyle::Get()->GetBrush(IconStyleAttr.Get()); }))
+				.ColorAndOpacity(FSlateColor::UseForeground())
+			];
+	}
+
+	inline void AppendButtons(TSharedRef<SHorizontalBox> InHorizBox, TArrayView<const FConcertActionDefinition> InDefs)
+	{
+		for (const FConcertActionDefinition& Def : InDefs)
+		{
+			InHorizBox->AddSlot()
+				.AutoWidth()
+				.VAlign(VAlign_Center)
+				.Padding(FMargin(1.0f))
+				[
+					Def.IconStyle.IsSet() ? CreateIconButton(Def) : CreateTextButton(Def)
+				];
+		}
+	}
+
+	template <typename ItemType, typename PredFactoryType>
+	inline void SyncArraysByPredicate(TArray<TSharedPtr<ItemType>>& InOutArray, TArray<TSharedPtr<ItemType>>&& InNewArray, const PredFactoryType& InPredFactory)
+	{
+		if (InOutArray.Num() == 0)
+		{
+			// Empty array - can just move
+			InOutArray = MoveTempIfPossible(InNewArray);
+		}
+		else
+		{
+			// Add or update the existing entries
+			for (TSharedPtr<ItemType>& NewItem : InNewArray)
+			{
+				TSharedPtr<ItemType>* ExistingItemPtr = InOutArray.FindByPredicate(InPredFactory(NewItem));
+				if (ExistingItemPtr)
+				{
+					**ExistingItemPtr = *NewItem;
+				}
+				else
+				{
+					InOutArray.Add(NewItem);
+				}
+			}
+			// Remove entries that are no longer needed
+			for (auto ExistingItemIt = InOutArray.CreateIterator(); ExistingItemIt; ++ExistingItemIt)
+			{
+				TSharedPtr<ItemType>* NewItemPtr = InNewArray.FindByPredicate(InPredFactory(*ExistingItemIt));
+				if (!NewItemPtr)
+				{
+					ExistingItemIt.RemoveCurrent();
+					continue;
+				}
+			}
+		}
+	}
+
+	template <typename ItemType>
+	inline TArray<TSharedPtr<ItemType>> DeepCopyArray(const TArray<TSharedPtr<ItemType>>& InArray)
+	{
+		TArray<TSharedPtr<ItemType>> ArrayCopy;
+		{
+			ArrayCopy.Reserve(InArray.Num());
+			for (const TSharedPtr<ItemType>& Item : InArray)
+			{
+				ArrayCopy.Add(MakeShared<ItemType>(*Item));
+			}
+		}
+		return ArrayCopy;
+	}
+
+	template <typename ItemType>
+	inline TArray<TSharedPtr<ItemType>> DeepCopyArrayAndClearSource(TArray<TSharedPtr<ItemType>>& InOutArray)
+	{
+		TArray<TSharedPtr<ItemType>> ArrayCopy = DeepCopyArray(InOutArray);
+		InOutArray.Reset();
+		return ArrayCopy;
+	}
+
+	/** Returns the image used to render the expandable area title bar with respect to its hover/expand state. */
+	inline const FSlateBrush* GetExpandableAreaBorderImage(const SExpandableArea& Area)
+	{
+		if (Area.IsTitleHovered())
+		{
+			return Area.IsExpanded() ? FAppStyle::Get().GetBrush("DetailsView.CategoryTop_Hovered") : FAppStyle::Get().GetBrush("DetailsView.CollapsedCategory_Hovered");
+		}
+		return Area.IsExpanded() ? FAppStyle::Get().GetBrush("DetailsView.CategoryTop") : FAppStyle::Get().GetBrush("DetailsView.CollapsedCategory");
+	}
+
+	static FText FormatRelativeTime(const FDateTime& EventTime, const FDateTime* CurrTime = nullptr)
+	{
+		FTimespan TimeSpan = (CurrTime ? *CurrTime : FDateTime::UtcNow()) - EventTime;
+		int32 Days = TimeSpan.GetDays();
+		int32 Hours = TimeSpan.GetHours();
+
+		if (Days >= 1)
+		{
+			return Hours > 0 ?
+				FText::Format(LOCTEXT("DaysHours", "{0} {0}|plural(one=Day,other=Days), {1} {1}|plural(one=Hour,other=Hours) Ago"), Days, Hours) :
+				FText::Format(LOCTEXT("Days", "{0} {0}|plural(one=Day,other=Days) Ago"), Days);
+		}
+
+		int32 Minutes = TimeSpan.GetMinutes();
+		if (Hours >= 1)
+		{
+			return Minutes > 0 ?
+				FText::Format(LOCTEXT("HoursMins", "{0} {0}|plural(one=Hour,other=Hours), {1} {1}|plural(one=Minute,other=Minutes) Ago"), Hours, Minutes) :
+				FText::Format(LOCTEXT("Hours", "{0} {0}|plural(one=Hour,other=Hours) Ago"), Hours);
+		}
+
+		int32 Seconds = TimeSpan.GetSeconds();
+		if (Minutes >= 1)
+		{
+			return Seconds > 0 ?
+				FText::Format(LOCTEXT("MinsSecs", "{0} {0}|plural(one=Minute,other=Minutes), {1} {1}|plural(one=Second,other=Seconds) Ago"), Minutes, Seconds) :
+				FText::Format(LOCTEXT("Mins", "{0} {0}|plural(one=Minute,other=Minutes) Ago"), Minutes, Seconds);
+		}
+
+		if (Seconds >= 1)
+		{
+			return FText::Format(LOCTEXT("Secs", "{0} {0}|plural(one=Second,other=Seconds) Ago"), Seconds);
+		}
+		return LOCTEXT("Now", "Now");
+	}
+};
+
+#undef LOCTEXT_NAMESPACE
diff --git a/Engine/Plugins/Developer/Concert/ConcertUI/ConcertSharedSlate/Source/ConcertSharedSlate/Public/IConcertSharedSlateModule.h b/Engine/Plugins/Developer/Concert/ConcertUI/ConcertSharedSlate/Source/ConcertSharedSlate/Public/IConcertSharedSlateModule.h
new file mode 100644
--- /dev/null
+++ b/Engine/Plugins/Developer/Concert/ConcertUI/ConcertSharedSlate/Source/ConcertSharedSlate/Public/IConcertSharedSlateModule.h
@@ -0,0 +1,21 @@
+// Copyright Epic Games, Inc. All Rights Reserved.
+
+#pragma once
+
+#include "CoreMinimal.h"
+#include "Modules/ModuleInterface.h"
+#include "Modules/ModuleManager.h"
+
+/**
+ * Interface for the Concert Sync UI module.
+ */
+class IConcertSharedSlateModule : public IModuleInterface
+{
+public:
+	/** Get the ConcertSyncUI module */
+	static IConcertSharedSlateModule& Get()
+	{
+		static const FName ModuleName = TEXT("ConcertSharedSlate");
+		return FModuleManager::Get().GetModuleChecked<IConcertSharedSlateModule>(ModuleName);
+	}
+};
diff --git a/Engine/Plugins/Developer/Concert/ConcertUI/ConcertSharedSlate/Source/ConcertSharedSlate/Public/SConcertScrollBox.h b/Engine/Plugins/Developer/Concert/ConcertUI/ConcertSharedSlate/Source/ConcertSharedSlate/Public/SConcertScrollBox.h
new file mode 100644
--- /dev/null
+++ b/Engine/Plugins/Developer/Concert/ConcertUI/ConcertSharedSlate/Source/ConcertSharedSlate/Public/SConcertScrollBox.h
@@ -0,0 +1,82 @@
+// Copyright Epic Games, Inc. All Rights Reserved.
+
+#pragma once
+
+#include "CoreMinimal.h"
+#include "Widgets/DeclarativeSyntaxSupport.h"
+#include "Widgets/SCompoundWidget.h"
+#include "Framework/SlateDelegates.h"
+
+class SScrollBox;
+class SButton;
+class SScrollBar;
+
+
+class CONCERTSHAREDSLATE_API SConcertScrollBox : public SCompoundWidget
+{
+public:
+
+	class FSlot : public TSlotBase<FSlot>
+	{
+	public:
+		using TSlotBase<FSlot>::TSlotBase;
+
+		SLATE_SLOT_BEGIN_ARGS(FSlot, TSlotBase<FSlot>)
+		SLATE_SLOT_END_ARGS()
+	};
+
+	SLATE_BEGIN_ARGS(SConcertScrollBox) {}
+		SLATE_SLOT_ARGUMENT(FSlot, Slots)
+	SLATE_END_ARGS();
+
+	/**
+	* Construct a custom ScrollBox that automatically scrolls to the bottom if the user is not currently scrolling.
+	*
+	* @param InArgs The Slate argument list.
+	* @param ConstructUnderMajorTab The major tab which will contain the Session History widget.
+	*/
+	void Construct(const FArguments& InArgs);
+
+	/** SWidget interface. */
+	virtual void Tick(const FGeometry& AllottedGeometry, const double InCurrentTime, const float InDeltaTime) override;
+	/** End of SWidget interface. */
+
+	/**
+	 * @return a new slot. Slots contain children for SScrollBox 
+	 */
+	static FSlot::FSlotArguments Slot();
+
+private:
+	
+	/** Delegate called to handle a user scrolling. */
+	void HandleUserScrolled(float);
+
+	/** Checks if the scrollbar is currently at the bottom of its track. */
+	bool IsAtBottom();
+
+	/** Callback for scrolling to the start of the list. */
+	FReply HandleScrollToStart();
+
+	/** Callback for scrolling to the end of the list. */
+	FReply HandleScrollToEnd();
+
+	/** Callback for getting the ScrollBar buttons visibility. */
+	EVisibility HandleScrollButtonsVisibility() const;
+
+	/** Creates a button with the FontAwesome font. */
+	TSharedRef<SButton> CreateScrollBarButton(const FText& InToolTip, const FString& InIcon, FOnClicked OnClickedDelegate) const;
+
+private:
+	
+	/** Whether the scrollbar is locked to the bottom. */
+	bool bIsLocked;
+
+	/** Whether locking should be prevented. */
+	bool bPreventLock;
+
+	/** Holds the ScrollBar. */
+	TSharedPtr<SScrollBar> ScrollBar;
+
+	/** Holds the inner ScrollBox. */
+	TSharedPtr<SScrollBox> ScrollBox;
+};
diff --git a/Engine/Plugins/Developer/Concert/ConcertUI/ConcertSharedSlate/Source/ConcertSharedSlate/Public/SConcertSessionActivities.h b/Engine/Plugins/Developer/Concert/ConcertUI/ConcertSharedSlate/Source/ConcertSharedSlate/Public/SConcertSessionActivities.h
new file mode 100644
--- /dev/null
+++ b/Engine/Plugins/Developer/Concert/ConcertUI/ConcertSharedSlate/Source/ConcertSharedSlate/Public/SConcertSessionActivities.h
@@ -0,0 +1,405 @@
+// Copyright Epic Games, Inc. All Rights Reserved.
+
+#pragma once
+
+#include "CoreMinimal.h"
+#include "SlateFwd.h"
+#include "Async/Future.h"
+#include "Misc/TextFilter.h"
+#include "Widgets/DeclarativeSyntaxSupport.h"
+#include "Widgets/SCompoundWidget.h"
+#include "Widgets/Views/SListView.h"
+#include "ConcertSyncSessionTypes.h"
+
+struct FConcertClientSessionActivity;
+struct FConcertClientInfo;
+struct FConcertPackageInfo;
+struct FConcertTransactionEventBase;
+class SUndoHistoryDetails;
+class SPackageDetails;
+class STextBlock;
+class SExpendableArea;
+
+
+/** Filters for the concert session activity view. */
+enum class EConcertActivityFilterFlags
+{
+	ShowAll                  = 0x00,
+	HideConnectionActivities = 1<<0,
+	HideLockActivities       = 1<<1,
+	HidePackageActivities    = 1<<2,
+	HideTransactionActivities= 1<<3,
+	HideIgnoredActivities    = 1<<4,
+};
+ENUM_CLASS_FLAGS(EConcertActivityFilterFlags);
+
+/** Displays session activities in a table view. */
+class CONCERTSHAREDSLATE_API SConcertSessionActivities : public SCompoundWidget
+{
+public:
+	/** Defines how the time should be displayed in the date/time column. */
+	enum class ETimeFormat
+	{
+		Relative, // Display relative time (23 seconds ago)
+		Absolute  // Display absolute time (April 7, 2019 - 10:33:52)
+	};
+
+	/** Used to pull activities from a session. Used to fetch and display the activities of an archived session. */
+	using FFetchActivitiesFunc = TFunction<bool(TArray<TSharedPtr<FConcertClientSessionActivity>>& /*InOutActivities*/, int32& /*OutFetchedCount*/, FText& /*ErrorMsg*/)>;
+
+	/** Used to map an activity to its client. */
+	using FGetActivityClientInfoFunc = TFunction<const FConcertClientInfo*(FGuid /*ClientId*/)>;
+
+	/** Returns the transaction event corresponding the specified activity.*/
+	using FGetTransactionEvent = TFunction<TFuture<TOptional<FConcertSyncTransactionEvent>>(const FConcertClientSessionActivity& /*Activity*/)>;
+
+	/** Returns the package event corresponding to the package activity. */
+	using FGetPackageEvent = TFunction<bool(const FConcertClientSessionActivity& /*Activity*/, FConcertSyncPackageEventMetaData& /*OutEvent*/)>;
+
+	/** Used to overlay a widget over a column widget to add custom functionalities to a row. */
+	using FMakeColumnOverlayWidgetFunc = TFunction<TSharedPtr<SWidget>(TWeakPtr<FConcertClientSessionActivity> /*ThisRowActivity*/, const FName& /*ColumnId*/)>;
+
+public:
+	SLATE_BEGIN_ARGS(SConcertSessionActivities)
+		: _OnFetchActivities()
+		, _OnMapActivityToClient()
+		, _OnMakeColumnOverlayWidget()
+		, _TimeFormat(ETimeFormat::Relative)
+		, _ClientAvatarColorColumnVisibility(EVisibility::Hidden)
+		, _ClientNameColumnVisibility(EVisibility::Hidden)
+		, _OperationColumnVisibility(EVisibility::Hidden)
+		, _PackageColumnVisibility(EVisibility::Hidden)
+		, _ConnectionActivitiesVisibility(EVisibility::Hidden)
+		, _LockActivitiesVisibility(EVisibility::Hidden)
+		, _PackageActivitiesVisibility(EVisibility::Visible)
+		, _TransactionActivitiesVisibility(EVisibility::Visible)
+		, _IgnoredActivitiesVisibility(EVisibility::Hidden)
+		, _DetailsAreaVisibility(EVisibility::Hidden)
+		, _IsAutoScrollEnabled(false){ }
+
+		/** If bound, invoked to populate the view. */
+		SLATE_ARGUMENT(FFetchActivitiesFunc, OnFetchActivities)
+
+		/** If bound, invoked to map an activity to a client.*/
+		SLATE_ARGUMENT(FGetActivityClientInfoFunc, OnMapActivityToClient)
+
+		/** If bound, invoked to fill up the package activity details panel. */
+		SLATE_ARGUMENT(FGetPackageEvent, OnGetPackageEvent)
+
+		/** If bound, invoked to fill up the transaction activity details panel. */
+		SLATE_ARGUMENT(FGetTransactionEvent, OnGetTransactionEvent)
+
+		/** If bound, invoked when generating a row to add an overlay to a column. */
+		SLATE_ARGUMENT(FMakeColumnOverlayWidgetFunc, OnMakeColumnOverlayWidget)
+
+		/** Highlight the returned text in the view. */
+		SLATE_ATTRIBUTE(FText, HighlightText)
+
+		/** Defines how time should be displayed (relative vs absolute). */
+		SLATE_ATTRIBUTE(ETimeFormat, TimeFormat)
+
+		/** Show/hide the column displaying the avatar color of the client who performed the activity. */
+		SLATE_ATTRIBUTE(EVisibility, ClientAvatarColorColumnVisibility)
+
+		/** Show/hide the column showing the display name of the client who performed the activity. */
+		SLATE_ATTRIBUTE(EVisibility, ClientNameColumnVisibility)
+
+		/** Show/hide the column showing the operation name. */
+		SLATE_ATTRIBUTE(EVisibility, OperationColumnVisibility)
+
+		/** Show/hide the column showing the affected package. */
+		SLATE_ATTRIBUTE(EVisibility, PackageColumnVisibility)
+
+		/** Show/hide connection activities. */
+		SLATE_ATTRIBUTE(EVisibility, ConnectionActivitiesVisibility)
+
+		/** Show/hide lock activities. */
+		SLATE_ATTRIBUTE(EVisibility, LockActivitiesVisibility)
+
+		/** Show/hide package activities. */
+		SLATE_ATTRIBUTE(EVisibility, PackageActivitiesVisibility)
+
+		/** Show/hide transaction activities. */
+		SLATE_ATTRIBUTE(EVisibility, TransactionActivitiesVisibility)
+
+		/** Show/hide ignored activities. */
+		SLATE_ATTRIBUTE(EVisibility, IgnoredActivitiesVisibility)
+
+		/** Show/hide the details area widget. (Not to confuse with widget expansion state) */
+		SLATE_ARGUMENT(EVisibility, DetailsAreaVisibility)
+
+		/** True to scroll the list down automatically (unless the user manually scrolled the list). */
+		SLATE_ARGUMENT(bool, IsAutoScrollEnabled)
+
+		/** Show/hide a message overlay above the activities list explaining why no activities are displayed. */
+		SLATE_ATTRIBUTE(FText, NoActivitiesReasonText)
+	SLATE_END_ARGS();
+
+	/**
+	 * Construct the recovery widget.
+	 * @param InArgs The widgets arguments and attributes.
+	 */
+	void Construct(const FArguments& InArgs);
+
+	/** Requests and consumes more activities if needed and/or possible. */
+	virtual void Tick(const FGeometry& AllottedGeometry, const double InCurrentTime, const float InDeltaTime) override;
+
+	/** Returns the activity selected or null if none is selected. */
+	TSharedPtr<FConcertClientSessionActivity> GetSelectedActivity() const;
+
+	/** Returns the total number of activities currently stored (no filter applied). */
+	int32 GetTotalActivityNum() const { return AllActivities.Num(); }
+
+	/** Returns the number of activities shown. */
+	int32 GetDisplayedActivityNum() const { return Activities.Num(); }
+
+	/** Returns the number of activities marked as 'ignored'. */
+	int32 GetIgnoredActivityNum() const { return IgnoredActivityNum; }
+
+	/** Returns the most recent activity available, ignoring the current filter. */
+	TSharedPtr<FConcertClientSessionActivity> GetMostRecentActivity() const;
+
+	/** Returns true if the column names is the last one (most right one). */
+	bool IsLastColumn(const FName& ColumnId) const;
+
+	/** Clears all activities displayed. */
+	void Reset();
+
+	/** Append an activity to the view. Used to populate the view from a live session. */
+	void Append(TSharedPtr<FConcertClientSessionActivity> Activity);
+
+	/** Request the view to refresh. */
+	void RequestRefresh();
+
+	/** Asks the view to update the text search filter. */
+	FText UpdateTextFilter(const FText& InFilterText);
+
+private:
+	/** Generates a row widget in the table view. */
+	TSharedRef<ITableRow> OnGenerateActivityRowWidget(TSharedPtr<FConcertClientSessionActivity> Item, const TSharedRef<STableViewBase>& OwnerTable);
+
+	/** Fetches more activities. */
+	void FetchActivities();
+
+	/** Queries the active set of activity filter flags. */
+	EConcertActivityFilterFlags QueryActiveActivityFilters() const;
+
+	/** Invoked when the filter changes and the displayed activity list must be reevaluated. */
+	void OnActivityFilterUpdated();
+
+	/** Returns true if the specified activity passes the view filters. */
+	bool PassesFilters(const FConcertClientSessionActivity& Activity);
+
+	/** Invoked when the list view scrolled to fetch more data from the activity provider. */
+	void OnListViewScrolled(double InScrollOffset);
+
+	/** Invoked when the list view selection changes. */
+	void OnListViewSelectionChanged(TSharedPtr<FConcertClientSessionActivity> InActivity, ESelectInfo::Type SelectInfo);
+
+	/** Convert the item into a list of text element used search bars. */
+	void PopulateSearchStrings(const FConcertClientSessionActivity& Item, TArray<FString>& OutSearchStrings) const;
+
+	/** Returns whether the details expandable area should be visible or not. */
+	EVisibility GetDetailAreaVisibility() const { return DetailsAreaVisibility; }
+
+	/** Manages how much space is used by the 'Details' area with respect to its expansion state. */
+	SSplitter::ESizeRule GetDetailsAreaSizeRule() const { return bDetailsAreaExpanded ? SSplitter::ESizeRule::FractionOfParent : SSplitter::ESizeRule::SizeToContent; }
+
+	/** Invoked when the 'Details' expandable area is expanded/collapsed. Affects how the 'Details' area size is accounted in the splitter. */
+	void OnDetailsAreaExpansionChanged(bool bExpanded);
+
+	/** Displays what properties/objects changed during a transaction in the details area when a transaction activity is selected in the list. */
+	void DisplayTransactionDetails(const FConcertClientSessionActivity& Activity, const FConcertTransactionEventBase& InTransaction);
+
+	/** Displays what changed in a package in the details area when a package activity is selected in the list. */
+	void DisplayPackageDetails(const FConcertClientSessionActivity& Activity, int64 PackageRevision, const FConcertPackageInfo& PackageInfo);
+
+	/** Update the detail area to display the proper panel and details. */
+	void UpdateDetailArea(TSharedPtr<FConcertClientSessionActivity> InSelectedActivity);
+
+	/** Change the visibility of the details panels, keeping only one visible. */
+	void SetDetailsPanelVisibility(const SWidget* VisiblePanel);
+
+	/** Returns the text displayed to the user when there are no details to display. */
+	FText GetNoDetailsText() const;
+
+	/** Invoked after all widget ticks get processed. Used to check if the user scrolled or not. */
+	void OnPostTick(float);
+
+private:
+	/** List of all activities (including the filtered out ones). */
+	TArray<TSharedPtr<FConcertClientSessionActivity>> AllActivities;
+
+	/** List of currently displayed activities. */
+	TArray<TSharedPtr<FConcertClientSessionActivity>> Activities;
+
+	/** The list view widget displaying the activities. */
+	TSharedPtr<SListView<TSharedPtr<FConcertClientSessionActivity>>> ActivityView;
+
+	/** Used to overlay a widget over a column widget (add an extra layer above the normal one) */
+	FMakeColumnOverlayWidgetFunc MakeColumnOverlayWidgetFn;
+
+	/** Returns which text should be highlighted. */
+	TAttribute<FText> HighlightText;
+
+	/** Whether the time should be displayed as relative (9 seconds ago) or as absolute (July 10, 2019 - 10:20:10) */
+	TAttribute<ETimeFormat> TimeFormat;
+
+	/** Whether the client name column is displayed. */
+	TAttribute<EVisibility> ClientNameColumnVisibility;
+
+	/** Whether the operation column is displayed. */
+	TAttribute<EVisibility> OperationColumnVisibility;
+
+	/** Whether the package column is displayed. */
+	TAttribute<EVisibility> PackageColumnVisibility;
+
+	/** Whether the join/leave session activities are displayed. */
+	TAttribute<EVisibility> ConnectionActivitiesVisibility;
+
+	/** Whether the lock/unlock activities are displayed.*/
+	TAttribute<EVisibility> LockActivitiesVisibility;
+
+	/** Whether the package activities are displayed.*/
+	TAttribute<EVisibility> PackageActivitiesVisibility;
+
+	/** Whether the transaction activities are displayed.*/
+	TAttribute<EVisibility> TransactionActivitiesVisibility;
+
+	/** Whether the ignored activities are displayed.*/
+	TAttribute<EVisibility> IgnoredActivitiesVisibility;
+
+	/** The number of activities flagged as 'ignored' within AllActivities array.*/
+	int32 IgnoredActivityNum = 0;
+
+	/** Whether the auto-scrolling is desired (scroll to bottom automatically unless the user scroll himself somewhere). */
+	bool bAutoScrollDesired = false;
+
+	/** Whether the user is scrolling (deactivate auto-scrolling). */
+	bool bUserScrolling = false;
+
+	/** Defines which activity types are currently filtered out from the view. */
+	EConcertActivityFilterFlags ActiveFilterFlags;
+
+	/** Used to fetch more activities from an abstract source. Usually mutually exclusive with Append() function. May not be bound. */
+	FFetchActivitiesFunc FetchActivitiesFn;
+
+	/** Used to map an activity endpoint ID to a client. May not be bound. */
+	FGetActivityClientInfoFunc GetActivityUserFn;
+
+	/** Used to get the transaction event to display the selected transaction activity details. May not be bound. */
+	FGetTransactionEvent GetTransactionEventFn;
+
+	/** Used to get the package event to display the selected package activity details. May not be bound. */
+	FGetPackageEvent GetPackageEventFn;
+
+	/** The number of activities to request when scrolling down to request activities on demand. Used with FetchActivitiesFn() for paging. */
+	static constexpr int32 ActivitiesPerRequest = 128;
+
+	/** The current desired amount of activities to display. Used for paging with FetchActivitiesFn(). It will grow when the user scrolls down. */
+	int32 DesiredActivitiesCount = ActivitiesPerRequest;
+
+	/** True once the activity provider function (FetchActivitiesFn) returns true. */
+	bool bAllActivitiesFetched = false;
+
+	/** Utility class used to tokenize and match text displayed in the list view. */
+	TSharedPtr<TTextFilter<const FConcertClientSessionActivity&>> SearchTextFilter;
+
+	/** The expandable area under which the activity details are displayed. */
+	TSharedPtr<SExpandableArea> ExpandableDetails;
+
+	/** The widget displaying transaction details (if detail area is expanded) when a transaction activity is selected. */
+	TSharedPtr<SUndoHistoryDetails> TransactionDetailsPanel;
+
+	/** The widget displaying package details (if detail area is expanded) when a package activity is selected.*/
+	TSharedPtr<SPackageDetails> PackageDetailsPanel;
+
+	/** The widget saying they are no detail available (if detail area is expanded) when the activity has no details or no activity is selected. */
+	TSharedPtr<SBox> NoDetailsPanel;
+
+	/** The widget displayed when details of a partially sync activity (if detail area is expanded) are being fetched from the server. */
+	TSharedPtr<SWidget> LoadingDetailsPanel;
+
+	/** Indicate if the details area should be displayed or not. */
+	EVisibility DetailsAreaVisibility = EVisibility::Collapsed;
+
+	/** Keeps the expanded status of the details area. */
+	bool bDetailsAreaExpanded = false;
+
+	/** Indicate whether the activity list view was scrolled during a frame. Used to correctly detect auto-scrolling in special cases. */
+	bool bActivityViewScrolled = false;
+};
+
+/**
+ * Manages the various options exposed by SConcertSessionActivities widget such as enabling/disabling filtering,
+ * changing the time format, etc.
+ */
+class CONCERTSHAREDSLATE_API FConcertSessionActivitiesOptions : public TSharedFromThis<FConcertSessionActivitiesOptions>
+{
+public:
+	/** Returns a Menu widgets containing the available options. */
+	TSharedRef<SWidget> MakeMenuWidget();
+
+	/** Makes a standard View Options widget, displaying the eye ball icon and showing the possible options. */
+	TSharedRef<SWidget> MakeViewOptionsWidget();
+
+	/** Makes a widgets saying how many operations are shown with respect to those filtered out. */
+	TSharedRef<SWidget> MakeDisplayedActivityCountWidget(TAttribute<int32> Total, TAttribute<int32> Displayed);
+
+	/** Makes a status bar widget displaying the activity shown and the standard view options button. */
+	TSharedRef<SWidget> MakeStatusBar(TAttribute<int32> Total, TAttribute<int32> Displayed);
+
+	/** Returns the time format option. */
+	SConcertSessionActivities::ETimeFormat GetTimeFormat() const { return bDisplayRelativeTime ? SConcertSessionActivities::ETimeFormat::Relative : SConcertSessionActivities::ETimeFormat::Absolute; }
+
+	/** Returns whether the connection activities are listed in SConcertSessionActivities. */
+	EVisibility GetConnectionActivitiesVisibility() const { return bDisplayConnectionActivities ? EVisibility::Visible : EVisibility::Hidden; }
+
+	/** Returns whether the lock activities are listed in SConcertSessionActivities. */
+	EVisibility GetLockActivitiesVisibility() const { return bDisplayLockActivities ? EVisibility::Visible : EVisibility::Hidden; }
+
+	/** Returns whether the package activities are listed in SConcertSessionActivities. */
+	EVisibility GetPackageActivitiesVisibility() const { return bDisplayPackageActivities ? EVisibility::Visible : EVisibility::Hidden; }
+
+	/** Returns whether the transaction activities are listed in SConcertSessionActivities. */
+	EVisibility GetTransactionActivitiesVisibility() const { return bDisplayTransactionActivities ? EVisibility::Visible : EVisibility::Hidden; }
+
+	/** Returns whether the ignored activities are listed in SConcertSessionActivities. */
+	EVisibility GetIgnoredActivitiesVisibility() const { return bDisplayIgnoredActivities ? EVisibility::Visible : EVisibility::Hidden; }
+
+	/** Invoked when an options is togged from the displayed menu widget. */
+	void OnOptionToggled(const FName CheckBoxId);
+
+	/** Enables the 'connection activity' filter check box. (Show Connection Activities). */
+	bool bEnableConnectionActivityFiltering = true;
+
+	/** Enables the 'lock activity' filter check box. (Show Lock Activities).*/
+	bool bEnableLockActivityFiltering = true;
+
+	/** Enables the 'package activity' filter check box (Show Package Activities). */
+	bool bEnablePackageActivityFiltering = true;
+
+	/** Enables the 'transaction activity' filter check box (Show Transaction Activities). */
+	bool bEnableTransactionActivityFiltering = true;
+
+	/** Enables the 'ignored activity' filter check box (Show Unrecoverable Activities). */
+	bool bEnableIgnoredActivityFiltering = false;
+
+	/** Controls whether the time is displayed as absolute or relative. (Display Relative Time). */
+	bool bDisplayRelativeTime = true;
+
+	/** If connection filtering is enabled, controls whether connection activities are filtered out. */
+	bool bDisplayConnectionActivities = false;
+
+	/** If lock filtering is enabled, controls whether lock activities are filtered out. */
+	bool bDisplayLockActivities = false;
+
+	/** If package filtering is enabled, controls whether package activities are filtered out. */
+	bool bDisplayPackageActivities = true;
+
+	/** If transaction filtering is enabled, controls whether transaction activities are filtered out.*/
+	bool bDisplayTransactionActivities = true;
+
+	/** If ignored activity filtering is enabled, controls whether ignored activities are filtered out.*/
+	bool bDisplayIgnoredActivities = false;
+};
diff --git a/Engine/Plugins/Developer/Concert/ConcertUI/ConcertSharedSlate/Source/ConcertSharedSlate/Public/SConcertSessionRecovery.h b/Engine/Plugins/Developer/Concert/ConcertUI/ConcertSharedSlate/Source/ConcertSharedSlate/Public/SConcertSessionRecovery.h
new file mode 100644
--- /dev/null
+++ b/Engine/Plugins/Developer/Concert/ConcertUI/ConcertSharedSlate/Source/ConcertSharedSlate/Public/SConcertSessionRecovery.h
@@ -0,0 +1,167 @@
+// Copyright Epic Games, Inc. All Rights Reserved.
+
+#pragma once
+
+#include "CoreMinimal.h"
+#include "SlateFwd.h"
+#include "Widgets/DeclarativeSyntaxSupport.h"
+#include "Widgets/SCompoundWidget.h"
+
+struct FConcertClientSessionActivity;
+struct FConcertClientInfo;
+class SConcertSessionActivities;
+class FConcertSessionActivitiesOptions;
+
+/**
+ * Displays the list of activities available for recovery and let the user select what should or shouldn't be recovered.
+ */
+class CONCERTSHAREDSLATE_API SConcertSessionRecovery : public SCompoundWidget
+{
+public:
+	/** Used to pull activities from a session. Used to fetch and display the activities of an archived session. */
+	using FFetchActivitiesFunc = TFunction<bool(TArray<TSharedPtr<FConcertClientSessionActivity>>& /*InOutActivities*/, int32& /*OutFetchedCount*/, FText& /*ErrorMsg*/)>;
+
+	/** Used to map an activity to its client. */
+	using FGetActivityClientInfoFunc = TFunction<const FConcertClientInfo*(FGuid /*ClientId*/)>;
+
+public:
+	SLATE_BEGIN_ARGS(SConcertSessionRecovery)
+		: _ClientAvatarColorColumnVisibility(EVisibility::Collapsed)
+		, _ClientNameColumnVisibility(EVisibility::Collapsed)
+		, _OperationColumnVisibility(EVisibility::Visible)
+		, _PackageColumnVisibility(EVisibility::Collapsed)
+		, _DetailsAreaVisibility(EVisibility::Collapsed)
+		, _IsConnectionActivityFilteringEnabled(false)
+		, _IsLockActivityFilteringEnabled(false)
+		, _IsPackageActivityFilteringEnabled(true)
+		, _IsTransactionActivityFilteringEnabled(true)
+		, _IsIgnoredActivityFilteringEnabled(false)
+		, _AreRecoverAllAndCancelButtonsVisible(true)
+		, _IsRecoverThroughButtonsVisible(true) {}
+
+		/** An introduction text to put the user in context. */
+		SLATE_ARGUMENT(FText, IntroductionText)
+
+		/** The windows hosting this widget. */
+		SLATE_ARGUMENT(TSharedPtr<SWindow>, ParentWindow)
+
+		/** If bound, invoked iteratively to populate the activity list. */
+		SLATE_ARGUMENT(FFetchActivitiesFunc, OnFetchActivities)
+
+		/** If bound, invoked to map an activity to a client info. */
+		SLATE_ARGUMENT(FGetActivityClientInfoFunc, OnMapActivityToClient)
+
+		/** Invoked when the user clicks the 'recover' button. */
+		SLATE_ARGUMENT(TFunction<bool(TSharedPtr<FConcertClientSessionActivity>)>, OnRestore)
+
+		/** Invoked when the user clicks the 'cancel' button. */
+		SLATE_ARGUMENT(TFunction<void()>, OnCancel)
+
+		/** Show/hide the column displaying the avatar color of the client who performed the activity. */
+		SLATE_ARGUMENT(EVisibility, ClientAvatarColorColumnVisibility)
+
+		/** Show/hide the column showing the display name of the client who performed the activity. */
+		SLATE_ARGUMENT(EVisibility, ClientNameColumnVisibility)
+
+		/** Show/hide the column displaying the operation represented by the activity. */
+		SLATE_ARGUMENT(EVisibility, OperationColumnVisibility)
+
+		/** Show/hide the column displaying affected package. */
+		SLATE_ARGUMENT(EVisibility, PackageColumnVisibility)
+
+		/** Show/hide the details area widget. */
+		SLATE_ARGUMENT(EVisibility, DetailsAreaVisibility)
+
+		/** Show/hide the check box in the 'View Options' to filter connection activities (join/leave session). */
+		SLATE_ARGUMENT(bool, IsConnectionActivityFilteringEnabled)
+
+		/** Show/hide the check box in the 'View Options' to filter lock activities (lock/unlock assets). */
+		SLATE_ARGUMENT(bool, IsLockActivityFilteringEnabled)
+
+		/** Show/hide the check box in the 'View Options' to filter package activities (create/delete/save/rename). */
+		SLATE_ARGUMENT(bool, IsPackageActivityFilteringEnabled)
+
+		/** Show/hide the check box in the 'View Options' to filter transaction activities. */
+		SLATE_ARGUMENT(bool, IsTransactionActivityFilteringEnabled)
+
+		/** Show/hide the check box in the 'View Options' to filter unrecoverable activities. */
+		SLATE_ARGUMENT(bool, IsIgnoredActivityFilteringEnabled)
+
+		/** Show/hide the recover all and cancel buttons (so that if this widget is embedded in another one, the buttons can be replaced and put in a more appropriate location). */
+		SLATE_ATTRIBUTE(bool, AreRecoverAllAndCancelButtonsVisible)
+
+		/** Show/hide the recover though button. */
+		SLATE_ATTRIBUTE(bool, IsRecoverThroughButtonsVisible)
+
+		/** Show/hide a message overlay above the activities list explaining why no activities are displayed. */
+		SLATE_ATTRIBUTE(FText, NoActivitiesReasonText)
+	SLATE_END_ARGS();
+
+	/**
+	 * Construct the recovery widget.
+	 * @param InArgs The widgets arguments and attributes.
+	 */
+	void Construct(const FArguments& InArgs);
+
+	/** Returns the activity, selected by the user, through which the session should be (or was) recovered or null to prevent recovery. */
+	TSharedPtr<FConcertClientSessionActivity> GetRecoverThroughItem() { return RecoveryThroughItem; }
+
+	/** Removes all activities from the view and reset the activity stream. */
+	void Reset();
+
+	/** Returns the total number of activities currently stored (no filter applied). */
+	int32 GetTotalActivityNum() const;
+
+	/** Returns the tooltip shown for the 'Recovery All' button. */
+	FText GetRecoverAllButtonTooltip() const;
+
+	/** Returns the most recent activity available, ignoring the current filter. */
+	TSharedPtr<FConcertClientSessionActivity> GetMostRecentActivity() const;
+
+private:
+	void OnSearchTextChanged(const FText& InFilterText);
+	void OnSearchTextCommitted(const FText& InFilterText, ETextCommit::Type CommitType);
+	FText HighlightSearchText() const;
+
+	FReply OnCancelRecoveryClicked();
+	FReply OnRecoverAllClicked();
+
+	EVisibility GetRecoverThroughButtonVisibility(TSharedPtr<FConcertClientSessionActivity> Activity);
+	FText GetRecoverThroughButtonTooltip() const;
+	TSharedPtr<SWidget> MakeRecoverThroughWidget(TWeakPtr<FConcertClientSessionActivity>, const FName&);
+	void RecoverThrough(TSharedPtr<FConcertClientSessionActivity> Item);
+
+	/** Close the windows hosting this recovery widget. */
+	void DismissWindow();
+
+private:
+	/** Display the session activities. */
+	TSharedPtr<SConcertSessionActivities> ActivityView;
+
+	/** Controls the various display options of the view. */
+	TSharedPtr<FConcertSessionActivitiesOptions> ActivityViewOptions;
+
+	/** The activity selected when the user click 'Recover' or 'Recover Through' buttons. */
+	TSharedPtr<FConcertClientSessionActivity> RecoveryThroughItem;
+
+	/** The parent window hosting this widget. */
+	TWeakPtr<SWindow> ParentWindow;
+
+	/** The widget used to enter the text to search. */
+	TSharedPtr<SSearchBox> SearchBox;
+
+	/** The search text entered in the search box. */
+	FText SearchText;
+
+	/** The text displayed at the top to summarize the purpose of the window. */
+	FText IntroductionText;
+
+	/** Externally control the 'recovery through' button visibility. The internal logic to show it is run first, if it is still visible, this is evaluated.*/
+	TAttribute<bool> IsRecoverThroughButtonVisible;
+
+	/** The function invoked when the user clicks the restore button. Might not be bound. */
+	TFunction<bool(TSharedPtr<FConcertClientSessionActivity>)> OnRestoreFn;
+
+	/** The function invoked when the user clicks the cancel button. Might not be bound. */
+	TFunction<void()> OnCancelFn;
+};
diff --git a/Engine/Plugins/Developer/Concert/ConcertUI/ConcertSharedSlate/Source/ConcertSharedSlate/Public/SPackageDetails.h b/Engine/Plugins/Developer/Concert/ConcertUI/ConcertSharedSlate/Source/ConcertSharedSlate/Public/SPackageDetails.h
new file mode 100644
--- /dev/null
+++ b/Engine/Plugins/Developer/Concert/ConcertUI/ConcertSharedSlate/Source/ConcertSharedSlate/Public/SPackageDetails.h
@@ -0,0 +1,57 @@
+// Copyright Epic Games, Inc. All Rights Reserved.
+
+#pragma once
+
+#include "CoreMinimal.h"
+#include "Widgets/DeclarativeSyntaxSupport.h"
+#include "Widgets/SCompoundWidget.h"
+#include "Widgets/Views/SListView.h"
+
+struct FConcertPackageInfo;
+
+class CONCERTSHAREDSLATE_API SPackageDetails : public SCompoundWidget
+{
+public:
+	SLATE_BEGIN_ARGS(SPackageDetails)
+	{}
+	SLATE_END_ARGS()
+
+	/**
+	* Constructs a package details widget.
+	*
+	* @param InArgs The Slate argument list.
+	*/
+	void Construct(const FArguments& InArgs);
+
+	/**
+	 * Set the current package to have its details displayed.
+	 * @param InPackageInfo The current package to have its info displayed.
+	 * @param InRevision The target package's revision.
+	 * @param InModifiedBy The username of the user who modified the package.
+	 */
+	void SetPackageInfo(const FConcertPackageInfo& InPackageInfo, const int64 InRevision, const FString& InModifiedBy);
+
+private:
+
+	/** Holds basic information about a given package. */
+	struct FPackageDetailsRow
+	{
+		FText PackageName;
+		FText Revision;
+		FText ModifiedBy;
+		FText Action;
+	};
+
+	using FListItemPtr = TSharedPtr<FPackageDetailsRow>;
+
+	/** Callback for generating package info rows. */
+	TSharedRef<ITableRow> HandleGenerateRow(FListItemPtr InItem, const TSharedRef<STableViewBase>& OwnerTable) const;
+
+private:
+
+	/** Holds the ListView of package details rows. */
+	TSharedPtr<SListView<FListItemPtr>>	DetailListView;
+
+	/** Holds the package details rows. */
+	TArray<FListItemPtr> DetailItems;
+};
\ No newline at end of file
