Fix non-unity compilation errors

diff --git a/Engine/Plugins/Developer/Concert/ConcertUI/ConcertSharedSlate/Source/ConcertSharedSlate/Private/SessionBrowser/SNewSessionRow.cpp b/Engine/Plugins/Developer/Concert/ConcertUI/ConcertSharedSlate/Source/ConcertSharedSlate/Private/SessionBrowser/SNewSessionRow.cpp
--- a/Engine/Plugins/Developer/Concert/ConcertUI/ConcertSharedSlate/Source/ConcertSharedSlate/Private/SessionBrowser/SNewSessionRow.cpp
+++ b/Engine/Plugins/Developer/Concert/ConcertUI/ConcertSharedSlate/Source/ConcertSharedSlate/Private/SessionBrowser/SNewSessionRow.cpp
@@ -6,9 +6,11 @@
 #include "SessionBrowser/ConcertSessionItem.h"
 
 #include "EditorFontGlyphs.h"
+#include "Framework/Application/SlateApplication.h"
 #include "Styling/AppStyle.h"
+#include "Widgets/Input/SEditableTextBox.h"
 #include "Widgets/Layout/SBox.h"
-#include "Widgets/Layout/SHeader.h"
+#include "Widgets/SBoxPanel.h"
 #include "Widgets/Layout/SUniformGridPanel.h"
 #include "Widgets/Text/STextBlock.h"
 #include "Widgets/Views/STableRow.h"
diff --git a/Engine/Plugins/Developer/Concert/ConcertUI/ConcertSharedSlate/Source/ConcertSharedSlate/Private/SessionBrowser/SNewSessionRow.h b/Engine/Plugins/Developer/Concert/ConcertUI/ConcertSharedSlate/Source/ConcertSharedSlate/Private/SessionBrowser/SNewSessionRow.h
--- a/Engine/Plugins/Developer/Concert/ConcertUI/ConcertSharedSlate/Source/ConcertSharedSlate/Private/SessionBrowser/SNewSessionRow.h
+++ b/Engine/Plugins/Developer/Concert/ConcertUI/ConcertSharedSlate/Source/ConcertSharedSlate/Private/SessionBrowser/SNewSessionRow.h
@@ -6,8 +6,11 @@
 #include "ConcertMessageData.h"
 #include "Widgets/DeclarativeSyntaxSupport.h"
 #include "Widgets/Views/STableRow.h"
+#include "Widgets/Input/SComboBox.h"
 
 class FConcertSessionItem;
+class SEditableTextBox;
+class SWidget;
 
 /**
  * The type of row used in the session list view to edit a new session (the session name + server).
diff --git a/Engine/Plugins/Developer/Concert/ConcertUI/ConcertSharedSlate/Source/ConcertSharedSlate/Private/SessionBrowser/SSaveRestoreSessionRow.cpp b/Engine/Plugins/Developer/Concert/ConcertUI/ConcertSharedSlate/Source/ConcertSharedSlate/Private/SessionBrowser/SSaveRestoreSessionRow.cpp
--- a/Engine/Plugins/Developer/Concert/ConcertUI/ConcertSharedSlate/Source/ConcertSharedSlate/Private/SessionBrowser/SSaveRestoreSessionRow.cpp
+++ b/Engine/Plugins/Developer/Concert/ConcertUI/ConcertSharedSlate/Source/ConcertSharedSlate/Private/SessionBrowser/SSaveRestoreSessionRow.cpp
@@ -8,8 +8,12 @@
 #include "EditorFontGlyphs.h"
 #include "Framework/Application/SlateApplication.h"
 #include "Internationalization/Regex.h"
-#include "Widgets/SBoxPanel.h"
+#include "Styling/AppStyle.h"
+#include "Widgets/Input/SEditableTextBox.h"
+#include "Widgets/Layout/SBox.h"
 #include "Widgets/Layout/SUniformGridPanel.h"
+#include "Widgets/SBoxPanel.h"
+#include "Widgets/Text/STextBlock.h"
 #include "Widgets/Views/SExpanderArrow.h"
 
 #define LOCTEXT_NAMESPACE "SConcertBrowser"
diff --git a/Engine/Plugins/Developer/Concert/ConcertUI/ConcertSharedSlate/Source/ConcertSharedSlate/Private/SessionBrowser/SSaveRestoreSessionRow.h b/Engine/Plugins/Developer/Concert/ConcertUI/ConcertSharedSlate/Source/ConcertSharedSlate/Private/SessionBrowser/SSaveRestoreSessionRow.h
--- a/Engine/Plugins/Developer/Concert/ConcertUI/ConcertSharedSlate/Source/ConcertSharedSlate/Private/SessionBrowser/SSaveRestoreSessionRow.h
+++ b/Engine/Plugins/Developer/Concert/ConcertUI/ConcertSharedSlate/Source/ConcertSharedSlate/Private/SessionBrowser/SSaveRestoreSessionRow.h
@@ -3,10 +3,13 @@
 #pragma once
 
 #include "CoreMinimal.h"
+#include "Misc/Attribute.h"
 #include "Widgets/DeclarativeSyntaxSupport.h"
 #include "Widgets/Views/STableRow.h"
 
 class FConcertSessionItem;
+class SWidget;
+class SEditableTextBox;
 
 /**
  * The type of row used in the session list view to create and archive or restore a session (edit the information required for the operation).
diff --git a/Engine/Plugins/Developer/Concert/ConcertUI/ConcertSharedSlate/Source/ConcertSharedSlate/Private/SessionBrowser/SSessionRow.cpp b/Engine/Plugins/Developer/Concert/ConcertUI/ConcertSharedSlate/Source/ConcertSharedSlate/Private/SessionBrowser/SSessionRow.cpp
--- a/Engine/Plugins/Developer/Concert/ConcertUI/ConcertSharedSlate/Source/ConcertSharedSlate/Private/SessionBrowser/SSessionRow.cpp
+++ b/Engine/Plugins/Developer/Concert/ConcertUI/ConcertSharedSlate/Source/ConcertSharedSlate/Private/SessionBrowser/SSessionRow.cpp
@@ -6,6 +6,13 @@
 #include "SessionBrowser/ConcertSessionItem.h"
 
 #include "EditorFontGlyphs.h"
+#include "Framework/Application/SlateApplication.h"
+#include "Styling/AppStyle.h"
+#include "Styling/CoreStyle.h"
+#include "Widgets/Layout/SBox.h"
+#include "Widgets/SBoxPanel.h"
+#include "Widgets/Text/SInlineEditableTextBlock.h"
+#include "Widgets/Text/STextBlock.h"
 
 #define LOCTEXT_NAMESPACE "SConcertBrowser"
 
diff --git a/Engine/Plugins/Developer/Concert/ConcertUI/ConcertSharedSlate/Source/ConcertSharedSlate/Private/SessionBrowser/SSessionRow.h b/Engine/Plugins/Developer/Concert/ConcertUI/ConcertSharedSlate/Source/ConcertSharedSlate/Private/SessionBrowser/SSessionRow.h
--- a/Engine/Plugins/Developer/Concert/ConcertUI/ConcertSharedSlate/Source/ConcertSharedSlate/Private/SessionBrowser/SSessionRow.h
+++ b/Engine/Plugins/Developer/Concert/ConcertUI/ConcertSharedSlate/Source/ConcertSharedSlate/Private/SessionBrowser/SSessionRow.h
@@ -3,11 +3,14 @@
 #pragma once
 
 #include "CoreMinimal.h"
+#include "Misc/Attribute.h"
 #include "Widgets/DeclarativeSyntaxSupport.h"
 #include "Widgets/Text/SInlineEditableTextBlock.h"
 #include "Widgets/Views/STableRow.h"
 
 class FConcertSessionItem;
+class SWidget;
+class SInlineEditableTextBlock;
 
 class SSessionRow : public SMultiColumnTableRow<TSharedPtr<FConcertSessionItem>>
 {
