Add missing forward declaration

diff --git a/Engine/Plugins/Developer/Concert/ConcertUI/ConcertSharedSlate/Source/ConcertSharedSlate/Public/SessionBrowser/SConcertSessionBrowser.h b/Engine/Plugins/Developer/Concert/ConcertUI/ConcertSharedSlate/Source/ConcertSharedSlate/Public/SessionBrowser/SConcertSessionBrowser.h
--- a/Engine/Plugins/Developer/Concert/ConcertUI/ConcertSharedSlate/Source/ConcertSharedSlate/Public/SessionBrowser/SConcertSessionBrowser.h
+++ b/Engine/Plugins/Developer/Concert/ConcertUI/ConcertSharedSlate/Source/ConcertSharedSlate/Public/SessionBrowser/SConcertSessionBrowser.h
@@ -11,6 +11,7 @@
 #include "Widgets/Views/SListView.h"
 
 class FConcertSessionItem;
+class FExtender;
 class IConcertSessionBrowserController;
 class ITableRow;
 class UConcertSessionBrowserSettings;
