Fix non-unity issues

diff --git a/Engine/Plugins/Developer/Concert/ConcertUI/ConcertServerUI/Source/ConcertServerUI/Private/Widgets/StatusBar/SConcertStatusBar.cpp b/Engine/Plugins/Developer/Concert/ConcertUI/ConcertServerUI/Source/ConcertServerUI/Private/Widgets/StatusBar/SConcertStatusBar.cpp
--- a/Engine/Plugins/Developer/Concert/ConcertUI/ConcertServerUI/Source/ConcertServerUI/Private/Widgets/StatusBar/SConcertStatusBar.cpp
+++ b/Engine/Plugins/Developer/Concert/ConcertUI/ConcertServerUI/Source/ConcertServerUI/Private/Widgets/StatusBar/SConcertStatusBar.cpp
@@ -2,6 +2,7 @@
 
 #include "SConcertStatusBar.h"
 
+#include "Framework/Application/SlateApplication.h"
 #include "OutputLogModule.h"
 #include "SWidgetDrawer.h"
 #include "Widgets/Input/SMultiLineEditableTextBox.h"
diff --git a/Engine/Plugins/Developer/Concert/ConcertUI/ConcertSharedSlate/Source/ConcertSharedSlate/Private/SessionBrowser/SNewSessionRow.cpp b/Engine/Plugins/Developer/Concert/ConcertUI/ConcertSharedSlate/Source/ConcertSharedSlate/Private/SessionBrowser/SNewSessionRow.cpp
--- a/Engine/Plugins/Developer/Concert/ConcertUI/ConcertSharedSlate/Source/ConcertSharedSlate/Private/SessionBrowser/SNewSessionRow.cpp
+++ b/Engine/Plugins/Developer/Concert/ConcertUI/ConcertSharedSlate/Source/ConcertSharedSlate/Private/SessionBrowser/SNewSessionRow.cpp
@@ -11,6 +11,7 @@
 #include "Widgets/Input/SEditableTextBox.h"
 #include "Widgets/Layout/SBox.h"
 #include "Widgets/SBoxPanel.h"
+#include "Widgets/Layout/SSpacer.h"
 #include "Widgets/Layout/SUniformGridPanel.h"
 #include "Widgets/Text/STextBlock.h"
 #include "Widgets/Views/STableRow.h"
diff --git a/Engine/Plugins/Developer/Concert/ConcertUI/ConcertSharedSlate/Source/ConcertSharedSlate/Private/SessionBrowser/SSaveRestoreSessionRow.h b/Engine/Plugins/Developer/Concert/ConcertUI/ConcertSharedSlate/Source/ConcertSharedSlate/Private/SessionBrowser/SSaveRestoreSessionRow.h
--- a/Engine/Plugins/Developer/Concert/ConcertUI/ConcertSharedSlate/Source/ConcertSharedSlate/Private/SessionBrowser/SSaveRestoreSessionRow.h
+++ b/Engine/Plugins/Developer/Concert/ConcertUI/ConcertSharedSlate/Source/ConcertSharedSlate/Private/SessionBrowser/SSaveRestoreSessionRow.h
@@ -4,10 +4,10 @@
 
 #include "CoreMinimal.h"
 #include "Misc/Attribute.h"
+#include "SessionBrowser/ConcertSessionItem.h"
 #include "Widgets/DeclarativeSyntaxSupport.h"
 #include "Widgets/Views/STableRow.h"
 
-class FConcertSessionItem;
 class SWidget;
 class SEditableTextBox;
 
