Remove ConcertSyncClient dependency from ConcertSharedSlate

diff --git a/Engine/Plugins/Developer/Concert/ConcertSync/ConcertSyncClient/Source/ConcertSyncClient/Private/ConcertActivityStream.cpp b/Engine/Plugins/Developer/Concert/ConcertSync/ConcertSyncClient/Source/ConcertSyncClient/Private/ConcertActivityStream.cpp
--- a/Engine/Plugins/Developer/Concert/ConcertSync/ConcertSyncClient/Source/ConcertSyncClient/Private/ConcertActivityStream.cpp
+++ b/Engine/Plugins/Developer/Concert/ConcertSync/ConcertSyncClient/Source/ConcertSyncClient/Private/ConcertActivityStream.cpp
@@ -20,7 +20,7 @@ FConcertActivityStream::FConcertActivityStream(TWeakPtr<IConcertClient, ESPMode:
 	Fetch(-GetRequestBatchSize()); // Negative -> request the N last activities (the most recent ones), up to 1024.
 }
 
-bool FConcertActivityStream::Read(TArray<TSharedPtr<FConcertClientSessionActivity>>& InOutArray, int32& OutReadCount, FText& OutErrorMsg)
+bool FConcertActivityStream::Read(TArray<TSharedPtr<FConcertSessionActivity>>& InOutArray, int32& OutReadCount, FText& OutErrorMsg)
 {
 	OutErrorMsg = LastErrorMsg;
 	if (!OutErrorMsg.IsEmpty())
@@ -114,7 +114,7 @@ void FConcertActivityStream::Fetch(int64 RequestFetchCount)
 					SyncActivity->EventSummary.GetPayload(EventSummaryPayload);
 					
 					// Note: Moving the 'SyncActivity' slices the derived type on purpose when this is a transaction/package activity. Details are decoupled in EventPayload.
-					AvailableActivities.Add(MakeShared<FConcertClientSessionActivity>(MoveTemp(*SyncActivity), MoveTemp(EventSummaryPayload), MoveTemp(EventPayload)));
+					AvailableActivities.Add(MakeShared<FConcertSessionActivity>(MoveTemp(*SyncActivity), MoveTemp(EventSummaryPayload), MoveTemp(EventPayload)));
 				}
 
 				if (RequestFetchCount > 0) // Response to second and subsequent requests.
diff --git a/Engine/Plugins/Developer/Concert/ConcertSync/ConcertSyncClient/Source/ConcertSyncClient/Private/ConcertClientWorkspace.cpp b/Engine/Plugins/Developer/Concert/ConcertSync/ConcertSyncClient/Source/ConcertSyncClient/Private/ConcertClientWorkspace.cpp
--- a/Engine/Plugins/Developer/Concert/ConcertSync/ConcertSyncClient/Source/ConcertSyncClient/Private/ConcertClientWorkspace.cpp
+++ b/Engine/Plugins/Developer/Concert/ConcertSync/ConcertSyncClient/Source/ConcertSyncClient/Private/ConcertClientWorkspace.cpp
@@ -404,7 +404,7 @@ bool FConcertClientWorkspace::IsTransactionEventPartiallySynced(const FConcertSy
 	return TransactionEvent.Transaction.ExportedObjects.Num() == 0;
 }
 
-void FConcertClientWorkspace::GetActivities(const int64 FirstActivityIdToFetch, const int64 MaxNumActivities, TMap<FGuid, FConcertClientInfo>& OutEndpointClientInfoMap, TArray<FConcertClientSessionActivity>& OutActivities) const
+void FConcertClientWorkspace::GetActivities(const int64 FirstActivityIdToFetch, const int64 MaxNumActivities, TMap<FGuid, FConcertClientInfo>& OutEndpointClientInfoMap, TArray<FConcertSessionActivity>& OutActivities) const
 {
 	OutEndpointClientInfoMap.Reset();
 	OutActivities.Reset();
diff --git a/Engine/Plugins/Developer/Concert/ConcertSync/ConcertSyncClient/Source/ConcertSyncClient/Private/ConcertClientWorkspace.h b/Engine/Plugins/Developer/Concert/ConcertSync/ConcertSyncClient/Source/ConcertSyncClient/Private/ConcertClientWorkspace.h
--- a/Engine/Plugins/Developer/Concert/ConcertSync/ConcertSyncClient/Source/ConcertSyncClient/Private/ConcertClientWorkspace.h
+++ b/Engine/Plugins/Developer/Concert/ConcertSync/ConcertSyncClient/Source/ConcertSyncClient/Private/ConcertClientWorkspace.h
@@ -47,7 +47,7 @@ public:
 	virtual bool FindTransactionEvent(const int64 TransactionEventId, FConcertSyncTransactionEvent& OutTransactionEvent, const bool bMetaDataOnly) const override;
 	virtual TFuture<TOptional<FConcertSyncTransactionEvent>> FindOrRequestTransactionEvent(const int64 TransactionEventId, const bool bMetaDataOnly) override;
 	virtual bool FindPackageEvent(const int64 PackageEventId, FConcertSyncPackageEventMetaData& OutPackageEvent) const override;
-	virtual void GetActivities(const int64 FirstActivityIdToFetch, const int64 MaxNumActivities, TMap<FGuid, FConcertClientInfo>& OutEndpointClientInfoMap, TArray<FConcertClientSessionActivity>& OutActivities) const override;
+	virtual void GetActivities(const int64 FirstActivityIdToFetch, const int64 MaxNumActivities, TMap<FGuid, FConcertClientInfo>& OutEndpointClientInfoMap, TArray<FConcertSessionActivity>& OutActivities) const override;
 	virtual int64 GetLastActivityId() const override;
 	virtual FOnActivityAddedOrUpdated& OnActivityAddedOrUpdated() override;
 	virtual FOnWorkspaceSynchronized& OnWorkspaceSynchronized() override;
diff --git a/Engine/Plugins/Developer/Concert/ConcertSync/ConcertSyncClient/Source/ConcertSyncClient/Public/ConcertActivityStream.h b/Engine/Plugins/Developer/Concert/ConcertSync/ConcertSyncClient/Source/ConcertSyncClient/Public/ConcertActivityStream.h
--- a/Engine/Plugins/Developer/Concert/ConcertSync/ConcertSyncClient/Source/ConcertSyncClient/Public/ConcertActivityStream.h
+++ b/Engine/Plugins/Developer/Concert/ConcertSync/ConcertSyncClient/Source/ConcertSyncClient/Public/ConcertActivityStream.h
@@ -5,7 +5,7 @@
 #include "CoreMinimal.h"
 #include "ConcertMessageData.h"
 
-struct FConcertClientSessionActivity;
+struct FConcertSessionActivity;
 class IConcertClient;
 
 /**
@@ -35,7 +35,7 @@ public:
 	 * @param OutErrorMsg Contains the last error message, if any or an empty text.
 	 * @return True if all activities were consumed, false otherwise.
 	 */
-	bool Read(TArray<TSharedPtr<FConcertClientSessionActivity>>& InOutActivities, int32& OutReadCount, FText& OutErrorMsg);
+	bool Read(TArray<TSharedPtr<FConcertSessionActivity>>& InOutActivities, int32& OutReadCount, FText& OutErrorMsg);
 
 	/** Returns the client info corresponding to an activities. The stream maintains this state as an optimization. */
 	const FConcertClientInfo* GetActivityClientInfo(const FGuid& ActivityEndpointId) const { return EndpointClientInfoMap.Find(ActivityEndpointId); }
@@ -53,7 +53,7 @@ private:
 	FGuid SessionId;
 	int64 LowestFetchedActivityId = 0; // The next Activity ID to fetch. Activity IDs are 1-based.
 	int64 TotalActivityCount = 0;
-	TArray<TSharedPtr<FConcertClientSessionActivity>> AvailableActivities; // Activities currently cached and ready for consumption.
+	TArray<TSharedPtr<FConcertSessionActivity>> AvailableActivities; // Activities currently cached and ready for consumption.
 	TMap<FGuid, FConcertClientInfo> EndpointClientInfoMap;
 	TSharedPtr<uint8, ESPMode::ThreadSafe> FetchRequestContinuationToken; // Used to disarm TFuture continuation.
 	FText LastErrorMsg;
diff --git a/Engine/Plugins/Developer/Concert/ConcertSync/ConcertSyncClient/Source/ConcertSyncClient/Public/IConcertClientWorkspace.h b/Engine/Plugins/Developer/Concert/ConcertSync/ConcertSyncClient/Source/ConcertSyncClient/Public/IConcertClientWorkspace.h
--- a/Engine/Plugins/Developer/Concert/ConcertSync/ConcertSyncClient/Source/ConcertSyncClient/Public/IConcertClientWorkspace.h
+++ b/Engine/Plugins/Developer/Concert/ConcertSync/ConcertSyncClient/Source/ConcertSyncClient/Public/IConcertClientWorkspace.h
@@ -18,41 +18,6 @@ DECLARE_DELEGATE_RetVal(bool, FCanProcessPendingPackages);
 DECLARE_MULTICAST_DELEGATE(FOnWorkspaceSynchronized);
 DECLARE_MULTICAST_DELEGATE_ThreeParams(FOnActivityAddedOrUpdated, const FConcertClientInfo&/*InClientInfo*/, const FConcertSyncActivity&/*InActivity*/, const FStructOnScope&/*InActivitySummary*/);
 
-
-struct FConcertClientSessionActivity
-{
-	FConcertClientSessionActivity() = default;
-
-	FConcertClientSessionActivity(const FConcertSyncActivity& InActivity, const FStructOnScope& InActivitySummary, TUniquePtr<FConcertSessionSerializedPayload> OptionalEventPayload = nullptr)
-		: Activity(InActivity)
-		, EventPayload(MoveTemp(OptionalEventPayload))
-	{
-		ActivitySummary.InitializeFromChecked(InActivitySummary);
-	}
-
-	FConcertClientSessionActivity(FConcertSyncActivity&& InActivity, FStructOnScope&& InActivitySummary, TUniquePtr<FConcertSessionSerializedPayload> OptionalEventPayload = nullptr)
-		: Activity(MoveTemp(InActivity))
-		, EventPayload(MoveTemp(OptionalEventPayload))
-	{
-		ActivitySummary.InitializeFromChecked(MoveTemp(InActivitySummary));
-	}
-
-	/** The generic activity part. */
-	FConcertSyncActivity Activity;
-
-	/** Contains the activity summary to display as text. */
-	TStructOnScope<FConcertSyncActivitySummary> ActivitySummary;
-
-	/**
-	 * The activity event payload usable for activity inspection. Might be null if it was not requested or did not provide insightful information.
-	 *   - If the activity type is 'transaction' and EventPayload is not null, it contains a FConcertSyncTransactionEvent with full transaction data.
-	 *   - If the activity type is 'package' and EventPayload is not null, it contains a FConcertSyncPackageEvent with the package meta data only.
-	 *   - Not set for other activity types (connection/lock).
-	 * @see FConcertActivityStream
-	 */
-	TUniquePtr<FConcertSessionSerializedPayload> EventPayload;
-};
-
 class IConcertClientWorkspace
 {
 public:
@@ -124,7 +89,7 @@ public:
 	 * @param OutEndpointClientInfoMap The client info for the activities fetched.
 	 * @param OutActivities the activities fetched.
 	 */
-	virtual void GetActivities(const int64 FirstActivityIdToFetch, const int64 MaxNumActivities, TMap<FGuid, FConcertClientInfo>& OutEndpointClientInfoMap, TArray<FConcertClientSessionActivity>& OutActivities) const = 0;
+	virtual void GetActivities(const int64 FirstActivityIdToFetch, const int64 MaxNumActivities, TMap<FGuid, FConcertClientInfo>& OutEndpointClientInfoMap, TArray<FConcertSessionActivity>& OutActivities) const = 0;
 
 	/**
 	 * Get the ID of the last activity in the session.
diff --git a/Engine/Plugins/Developer/Concert/ConcertSync/ConcertSyncCore/Source/ConcertSyncCore/Public/ConcertSyncSessionTypes.h b/Engine/Plugins/Developer/Concert/ConcertSync/ConcertSyncCore/Source/ConcertSyncCore/Public/ConcertSyncSessionTypes.h
--- a/Engine/Plugins/Developer/Concert/ConcertSync/ConcertSyncCore/Source/ConcertSyncCore/Public/ConcertSyncSessionTypes.h
+++ b/Engine/Plugins/Developer/Concert/ConcertSync/ConcertSyncCore/Source/ConcertSyncCore/Public/ConcertSyncSessionTypes.h
@@ -6,6 +6,7 @@
 #include "ConcertMessageData.h"
 #include "ConcertWorkspaceData.h"
 #include "ConcertTransactionEvents.h"
+#include "UObject/StructOnScope.h"
 #include "ConcertSyncSessionTypes.generated.h"
 
 /** Types of connection events */
@@ -405,3 +406,37 @@ protected:
 	virtual FText CreateDisplayText(const bool InUseRichText) const override;
 	virtual FText CreateDisplayTextForUser(const FText InUserDisplayName, const bool InUseRichText) const override;
 };
+
+struct FConcertSessionActivity
+{
+	FConcertSessionActivity() = default;
+
+	FConcertSessionActivity(const FConcertSyncActivity& InActivity, const FStructOnScope& InActivitySummary, TUniquePtr<FConcertSessionSerializedPayload> OptionalEventPayload = nullptr)
+		: Activity(InActivity)
+		, EventPayload(MoveTemp(OptionalEventPayload))
+	{
+		ActivitySummary.InitializeFromChecked(InActivitySummary);
+	}
+
+	FConcertSessionActivity(FConcertSyncActivity&& InActivity, FStructOnScope&& InActivitySummary, TUniquePtr<FConcertSessionSerializedPayload> OptionalEventPayload = nullptr)
+		: Activity(MoveTemp(InActivity))
+		, EventPayload(MoveTemp(OptionalEventPayload))
+	{
+		ActivitySummary.InitializeFromChecked(MoveTemp(InActivitySummary));
+	}
+
+	/** The generic activity part. */
+	FConcertSyncActivity Activity;
+
+	/** Contains the activity summary to display as text. */
+	TStructOnScope<FConcertSyncActivitySummary> ActivitySummary;
+
+	/**
+	 * The activity event payload usable for activity inspection. Might be null if it was not requested or did not provide insightful information.
+	 *   - If the activity type is 'transaction' and EventPayload is not null, it contains a FConcertSyncTransactionEvent with full transaction data.
+	 *   - If the activity type is 'package' and EventPayload is not null, it contains a FConcertSyncPackageEvent with the package meta data only.
+	 *   - Not set for other activity types (connection/lock).
+	 * @see FConcertActivityStream
+	 */
+	TUniquePtr<FConcertSessionSerializedPayload> EventPayload;
+};
\ No newline at end of file
diff --git a/Engine/Plugins/Developer/Concert/ConcertUI/ConcertClientSharedSlate/ConcertClientSharedSlate.uplugin b/Engine/Plugins/Developer/Concert/ConcertUI/ConcertClientSharedSlate/ConcertClientSharedSlate.uplugin
new file mode 100644
--- /dev/null
+++ b/Engine/Plugins/Developer/Concert/ConcertUI/ConcertClientSharedSlate/ConcertClientSharedSlate.uplugin
@@ -0,0 +1,41 @@
+{
+	"FileVersion": 3,
+	"Version": 1,
+	"VersionName": "0.1",
+	"FriendlyName": "Concert Shared Slate",
+	"Description": "Contains UI that is shared by client UI modules only",
+	"Category": "Networking",
+	"CreatedBy": "Epic Games, Inc.",
+	"CreatedByURL": "http://epicgames.com",
+	"DocsURL": "",
+	"MarketplaceURL": "",
+	"SupportURL": "",
+	"EnabledByDefault": false,
+	"CanContainContent": false,
+	"Hidden": true,
+	"IsBetaVersion": true,
+	"Installed": false,
+	"SupportedPrograms": [ "UnrealMultiUserServer", "UnrealRecoverySvc", "CrashReportClientEditor" ],
+	"Modules": [
+        {
+          "Name": "ConcertClientSharedSlate",
+          "Type": "UncookedOnly",
+          "LoadingPhase": "Default",
+          "ProgramAllowList": [ "UnrealMultiUserServer", "UnrealRecoverySvc", "CrashReportClientEditor" ]
+        }
+	],
+	"Plugins": [
+		{
+			"Name": "ConcertMain",
+			"Enabled": true
+		},
+		{
+			"Name": "ConcertSyncCore",
+			"Enabled": true
+		},
+        {
+          "Name": "ConcertSyncClient",
+          "Enabled": true
+        }
+	]
+}
diff --git a/Engine/Plugins/Developer/Concert/ConcertUI/ConcertClientSharedSlate/Source/ConcertClientSharedSlate/ConcertClientSharedSlate.Build.cs b/Engine/Plugins/Developer/Concert/ConcertUI/ConcertClientSharedSlate/Source/ConcertClientSharedSlate/ConcertClientSharedSlate.Build.cs
new file mode 100644
--- /dev/null
+++ b/Engine/Plugins/Developer/Concert/ConcertUI/ConcertClientSharedSlate/Source/ConcertClientSharedSlate/ConcertClientSharedSlate.Build.cs
@@ -0,0 +1,23 @@
+// Copyright Epic Games, Inc. All Rights Reserved.
+
+namespace UnrealBuildTool.Rules
+{
+	public class ConcertClientSharedSlate : ModuleRules
+	{
+		public ConcertClientSharedSlate(ReadOnlyTargetRules Target) : base(Target)
+		{
+			PublicDependencyModuleNames.AddRange(
+				new string[]
+				{
+					"Core",
+					"CoreUObject",
+					"Slate",
+					"SlateCore",
+					"EditorStyle",
+					"ConcertSharedSlate",
+					"ConcertSyncClient",
+				}
+			);
+		}
+	}
+}
diff --git a/Engine/Plugins/Developer/Concert/ConcertUI/ConcertClientSharedSlate/Source/ConcertClientSharedSlate/Public/ConcertClientFrontendUtils.h b/Engine/Plugins/Developer/Concert/ConcertUI/ConcertClientSharedSlate/Source/ConcertClientSharedSlate/Public/ConcertClientFrontendUtils.h
new file mode 100644
--- /dev/null
+++ b/Engine/Plugins/Developer/Concert/ConcertUI/ConcertClientSharedSlate/Source/ConcertClientSharedSlate/Public/ConcertClientFrontendUtils.h
@@ -0,0 +1,92 @@
+// Copyright Epic Games, Inc. All Rights Reserved.
+
+#pragma once
+
+#include "CoreMinimal.h"
+#include "ConcertActionDefinition.h"
+#include "Styling/AppStyle.h"
+#include "Widgets/Layout/SBorder.h"
+#include "Widgets/Images/SImage.h"
+#include "Widgets/Input/SButton.h"
+#include "Widgets/Text/STextBlock.h"
+#include "Widgets/Layout/SBox.h"
+#include "Widgets/Layout/SExpandableArea.h"
+#include "Widgets/SBoxPanel.h"
+#include "ConcertFrontendStyle.h"
+
+#define LOCTEXT_NAMESPACE "ConcertClientFrontendUtils"
+
+namespace ConcertClientFrontendUtils
+{
+	static const FName ButtonIconSyle = TEXT("FontAwesome.10");
+	static const float MinDesiredWidthForBtnAndIcon = 29.f;
+	static const FName ButtonStyleNames[(int32)EConcertActionType::NUM] = {
+		TEXT("SimpleButton"),
+		TEXT("FlatButton.Primary"),
+		TEXT("FlatButton.Info"),
+		TEXT("FlatButton.Success"),
+		TEXT("FlatButton.Warning"),
+		TEXT("FlatButton.Danger"),
+	};
+	
+	inline TSharedRef<SButton> CreateTextButton(const FConcertActionDefinition& InDef)
+	{
+		const FButtonStyle* ButtonStyle = &FAppStyle::Get().GetWidgetStyle<FButtonStyle>(ButtonStyleNames[(int32)InDef.Type]);
+		check(ButtonStyle);
+		const float ButtonContentWidthPadding = 6.f;
+		const float PaddingCompensation = (ButtonStyle->NormalPadding.Left + ButtonStyle->NormalPadding.Right + ButtonContentWidthPadding * 2);
+
+		return SNew(SButton)
+			.ButtonStyle(ButtonStyle)
+			.ToolTipText(InDef.ToolTipText)
+			.ContentPadding(FMargin(ButtonContentWidthPadding, 2.f))
+			.IsEnabled(InDef.IsEnabled)
+			.Visibility_Lambda([IsVisible = InDef.IsVisible]() { return IsVisible.Get() ? EVisibility::Visible : EVisibility::Collapsed; })
+			.OnClicked_Lambda([OnExecute = InDef.OnExecute]() { OnExecute.ExecuteIfBound(); return FReply::Handled(); })
+			[
+				SNew(SBox)
+				.MinDesiredWidth(MinDesiredWidthForBtnAndIcon - PaddingCompensation)
+				.HAlign(HAlign_Center)
+				[
+					SNew(STextBlock)
+					.Font(FAppStyle::Get().GetFontStyle(ButtonIconSyle))
+					.Text(InDef.Text)
+					.Justification(ETextJustify::Center)
+				]
+			];
+	}
+
+	inline TSharedRef<SButton> CreateIconButton(const FConcertActionDefinition& InDef)
+	{
+		return SNew(SButton)
+			.ButtonStyle(FAppStyle::Get(), ButtonStyleNames[(int32)InDef.Type])
+			.ToolTipText(InDef.ToolTipText)
+			.ContentPadding(FMargin(0, 0))
+			.IsEnabled(InDef.IsEnabled)
+			.Visibility_Lambda([IsVisible = InDef.IsVisible]() { return IsVisible.Get() ? EVisibility::Visible : EVisibility::Collapsed; })
+			.OnClicked_Lambda([OnExecute = InDef.OnExecute]() { OnExecute.ExecuteIfBound(); return FReply::Handled(); })
+			.VAlign(VAlign_Center)
+			.HAlign(HAlign_Center)
+			[
+				SNew(SImage)
+				.Image(TAttribute<const FSlateBrush*>::Create([IconStyleAttr = InDef.IconStyle]() { return FConcertFrontendStyle::Get()->GetBrush(IconStyleAttr.Get()); }))
+				.ColorAndOpacity(FSlateColor::UseForeground())
+			];
+	}
+
+	inline void AppendButtons(TSharedRef<SHorizontalBox> InHorizBox, TArrayView<const FConcertActionDefinition> InDefs)
+	{
+		for (const FConcertActionDefinition& Def : InDefs)
+		{
+			InHorizBox->AddSlot()
+				.AutoWidth()
+				.VAlign(VAlign_Center)
+				.Padding(FMargin(1.0f))
+				[
+					Def.IconStyle.IsSet() ? CreateIconButton(Def) : CreateTextButton(Def)
+				];
+		}
+	}
+};
+
+#undef LOCTEXT_NAMESPACE
diff --git a/Engine/Plugins/Developer/Concert/ConcertUI/ConcertSharedSlate/ConcertSharedSlate.uplugin b/Engine/Plugins/Developer/Concert/ConcertUI/ConcertSharedSlate/ConcertSharedSlate.uplugin
--- a/Engine/Plugins/Developer/Concert/ConcertUI/ConcertSharedSlate/ConcertSharedSlate.uplugin
+++ b/Engine/Plugins/Developer/Concert/ConcertUI/ConcertSharedSlate/ConcertSharedSlate.uplugin
@@ -2,8 +2,8 @@
 	"FileVersion": 3,
 	"Version": 1,
 	"VersionName": "0.2",
-	"FriendlyName": "Concert Sync Server UI",
-	"Description": "UI plugin that visualizes the multi-user server",
+	"FriendlyName": "Concert Shared Slate",
+	"Description": "Contains UI that is shared for server and client UI modules",
 	"Category": "Networking",
 	"CreatedBy": "Epic Games, Inc.",
 	"CreatedByURL": "http://epicgames.com",
@@ -32,14 +32,6 @@
 		{
 			"Name": "ConcertSyncCore",
 			"Enabled": true
-		},
-        {
-          "Name": "ConcertSyncClient",
-          "Enabled": true
-        },
-		{
-			"Name": "ConcertSyncServer",
-			"Enabled": true
 		}
 	]
 }
diff --git a/Engine/Plugins/Developer/Concert/ConcertUI/ConcertSharedSlate/Source/ConcertSharedSlate/ConcertSharedSlate.Build.cs b/Engine/Plugins/Developer/Concert/ConcertUI/ConcertSharedSlate/Source/ConcertSharedSlate/ConcertSharedSlate.Build.cs
--- a/Engine/Plugins/Developer/Concert/ConcertUI/ConcertSharedSlate/Source/ConcertSharedSlate/ConcertSharedSlate.Build.cs
+++ b/Engine/Plugins/Developer/Concert/ConcertUI/ConcertSharedSlate/Source/ConcertSharedSlate/ConcertSharedSlate.Build.cs
@@ -25,7 +25,6 @@ namespace UnrealBuildTool.Rules
 					"Projects",
 					"Concert",
 					"ConcertTransport",
-					"ConcertSyncClient",
 					"ConcertSyncCore",
 					"UndoHistory"
 				}
diff --git a/Engine/Plugins/Developer/Concert/ConcertUI/ConcertSharedSlate/Source/ConcertSharedSlate/Private/SConcertSessionActivities.cpp b/Engine/Plugins/Developer/Concert/ConcertUI/ConcertSharedSlate/Source/ConcertSharedSlate/Private/SConcertSessionActivities.cpp
--- a/Engine/Plugins/Developer/Concert/ConcertUI/ConcertSharedSlate/Source/ConcertSharedSlate/Private/SConcertSessionActivities.cpp
+++ b/Engine/Plugins/Developer/Concert/ConcertUI/ConcertSharedSlate/Source/ConcertSharedSlate/Private/SConcertSessionActivities.cpp
@@ -2,7 +2,6 @@
 
 #include "SConcertSessionActivities.h"
 
-#include "IConcertClientWorkspace.h"
 #include "Algo/Transform.h"
 #include "Misc/AsyncTaskNotification.h"
 #include "Misc/ITransaction.h"
@@ -50,12 +49,12 @@ const FName ShowPackageActivitiesCheckBoxId     = TEXT("ShowPackageActivities");
 const FName ShowTransactionActivitiesCheckBoxId = TEXT("ShowTransactionActivities");
 const FName ShowIgnoredActivitiesCheckBoxId     = TEXT("ShowIgnoredActivities");
 
-FText GetActivityDateTime(const FConcertClientSessionActivity& Activity, SConcertSessionActivities::ETimeFormat TimeFormat)
+FText GetActivityDateTime(const FConcertSessionActivity& Activity, SConcertSessionActivities::ETimeFormat TimeFormat)
 {
 	return TimeFormat == SConcertSessionActivities::ETimeFormat::Relative ? ConcertFrontendUtils::FormatRelativeTime(Activity.Activity.EventTime) : FText::AsDateTime(Activity.Activity.EventTime);
 }
 
-FText GetOperationName(const FConcertClientSessionActivity& Activity)
+FText GetOperationName(const FConcertSessionActivity& Activity)
 {
 	if (const FConcertSyncTransactionActivitySummary* TransactionSummary = Activity.ActivitySummary.Cast<FConcertSyncTransactionActivitySummary>())
 	{
@@ -107,7 +106,7 @@ FText GetOperationName(const FConcertClientSessionActivity& Activity)
 	return FText::GetEmpty();
 }
 
-FText GetPackageName(const FConcertClientSessionActivity& Activity)
+FText GetPackageName(const FConcertSessionActivity& Activity)
 {
 	if (const FConcertSyncPackageActivitySummary* Summary = Activity.ActivitySummary.Cast<FConcertSyncPackageActivitySummary>())
 	{
@@ -122,7 +121,7 @@ FText GetPackageName(const FConcertClientSessionActivity& Activity)
 	return FText::GetEmpty();
 }
 
-FText GetSummary(const FConcertClientSessionActivity& Activity, const FText& ClientName, bool bAsRichText)
+FText GetSummary(const FConcertSessionActivity& Activity, const FText& ClientName, bool bAsRichText)
 {
 	if (const FConcertSyncActivitySummary* Summary = Activity.ActivitySummary.Cast<FConcertSyncActivitySummary>())
 	{
@@ -142,7 +141,7 @@ FText GetClientName(const FConcertClientInfo* InActivityClient)
 /**
  * Displays the summary of an activity recorded and recoverable in the SConcertSessionRecovery list view.
  */
-class SConcertSessionActivityRow : public SMultiColumnTableRow<TSharedPtr<FConcertClientSessionActivity>>
+class SConcertSessionActivityRow : public SMultiColumnTableRow<TSharedPtr<FConcertSessionActivity>>
 {
 public:
 	SLATE_BEGIN_ARGS(SConcertSessionActivityRow)
@@ -166,7 +165,7 @@ public:
 	 * @param InActivityClient The client who produced this activity. Can be null if unknown or not desirable.
 	 * @param InOwnerTableView The table view that will own this row.
 	 */
-	void Construct(const FArguments& InArgs, TSharedPtr<FConcertClientSessionActivity> InActivity, const FConcertClientInfo* InActivityClient, const TSharedRef<STableViewBase>& InOwnerTableView);
+	void Construct(const FArguments& InArgs, TSharedPtr<FConcertSessionActivity> InActivity, const FConcertClientInfo* InActivityClient, const TSharedRef<STableViewBase>& InOwnerTableView);
 
 	/** Generates the widget representing this row. */
 	virtual TSharedRef<SWidget> GenerateWidgetForColumn(const FName& ColumnName) override;
@@ -185,7 +184,7 @@ private:
 	FText MakeTooltipText() const;
 
 private:
-	TWeakPtr<FConcertClientSessionActivity> Activity;
+	TWeakPtr<FConcertSessionActivity> Activity;
 	TAttribute<SConcertSessionActivities::ETimeFormat> TimeFormat;
 	FText AbsoluteDateTime;
 	FText ClientName;
@@ -195,7 +194,7 @@ private:
 };
 
 
-void SConcertSessionActivityRow::Construct(const FArguments& InArgs, TSharedPtr<FConcertClientSessionActivity> InActivity, const FConcertClientInfo* InActivityClient, const TSharedRef<STableViewBase>& InOwnerTableView)
+void SConcertSessionActivityRow::Construct(const FArguments& InArgs, TSharedPtr<FConcertSessionActivity> InActivity, const FConcertClientInfo* InActivityClient, const TSharedRef<STableViewBase>& InOwnerTableView)
 {
 	Activity = InActivity;
 	TimeFormat = InArgs._TimeFormat;
@@ -206,7 +205,7 @@ void SConcertSessionActivityRow::Construct(const FArguments& InArgs, TSharedPtr<
 	ClientAvatarColor = InActivityClient ? InActivityClient->AvatarColor : FConcertFrontendStyle::Get()->GetColor("Concert.DisconnectedColor");
 
 	// Construct base class
-	SMultiColumnTableRow<TSharedPtr<FConcertClientSessionActivity>>::Construct(FSuperRowType::FArguments(), InOwnerTableView);
+	SMultiColumnTableRow<TSharedPtr<FConcertSessionActivity>>::Construct(FSuperRowType::FArguments(), InOwnerTableView);
 
 	if (InActivity->Activity.bIgnored)
 	{
@@ -216,7 +215,7 @@ void SConcertSessionActivityRow::Construct(const FArguments& InArgs, TSharedPtr<
 
 TSharedRef<SWidget> SConcertSessionActivityRow::GenerateWidgetForColumn(const FName& ColumnId)
 {
-	TSharedPtr<FConcertClientSessionActivity> ActivityPin = Activity.Pin();
+	TSharedPtr<FConcertSessionActivity> ActivityPin = Activity.Pin();
 	TSharedRef<SOverlay> Overlay = SNew(SOverlay);
 
 	if (ColumnId == ConcertSessionActivityUtils::AvatarColorColumnId)
@@ -303,7 +302,7 @@ TSharedRef<SWidget> SConcertSessionActivityRow::GenerateWidgetForColumn(const FN
 
 FText SConcertSessionActivityRow::MakeTooltipText() const
 {
-	if (TSharedPtr<FConcertClientSessionActivity> ActivityPin = Activity.Pin())
+	if (TSharedPtr<FConcertSessionActivity> ActivityPin = Activity.Pin())
 	{
 		FText Client = GetClientName();
 		FText Operation = ConcertSessionActivityUtils::GetOperationName(*ActivityPin);
@@ -343,7 +342,7 @@ FText SConcertSessionActivityRow::MakeTooltipText() const
 
 FText SConcertSessionActivityRow::FormatEventDateTime() const
 {
-	if (TSharedPtr<FConcertClientSessionActivity> ItemPin = Activity.Pin())
+	if (TSharedPtr<FConcertSessionActivity> ItemPin = Activity.Pin())
 	{
 		return TimeFormat.Get() == SConcertSessionActivities::ETimeFormat::Relative ? ConcertFrontendUtils::FormatRelativeTime(ItemPin->Activity.EventTime) : AbsoluteDateTime;
 	}
@@ -371,7 +370,7 @@ void SConcertSessionActivities::Construct(const FArguments& InArgs)
 	DetailsAreaVisibility = InArgs._DetailsAreaVisibility;
 	bAutoScrollDesired = InArgs._IsAutoScrollEnabled;
 
-	SearchTextFilter = MakeShared<TTextFilter<const FConcertClientSessionActivity&>>(TTextFilter<const FConcertClientSessionActivity&>::FItemToStringArray::CreateSP(this, &SConcertSessionActivities::PopulateSearchStrings));
+	SearchTextFilter = MakeShared<TTextFilter<const FConcertSessionActivity&>>(TTextFilter<const FConcertSessionActivity&>::FItemToStringArray::CreateSP(this, &SConcertSessionActivities::PopulateSearchStrings));
 	SearchTextFilter->OnChanged().AddSP(this, &SConcertSessionActivities::OnActivityFilterUpdated);
 
 	// Set the initial filter state.
@@ -432,7 +431,7 @@ void SConcertSessionActivities::Construct(const FArguments& InArgs)
 				.BorderBackgroundColor(FSlateColor(FLinearColor(0.6, 0.6, 0.6)))
 				.Padding(0)
 				[
-					SAssignNew(ActivityView, SListView<TSharedPtr<FConcertClientSessionActivity>>)
+					SAssignNew(ActivityView, SListView<TSharedPtr<FConcertSessionActivity>>)
 					.ListItemsSource(&Activities)
 					.OnGenerateRow(this, &SConcertSessionActivities::OnGenerateActivityRowWidget)
 					.SelectionMode(ESelectionMode::Single)
@@ -582,7 +581,7 @@ void SConcertSessionActivities::OnListViewScrolled(double InScrollOffset)
 	}
 }
 
-void SConcertSessionActivities::OnListViewSelectionChanged(TSharedPtr<FConcertClientSessionActivity> InActivity, ESelectInfo::Type SelectInfo)
+void SConcertSessionActivities::OnListViewSelectionChanged(TSharedPtr<FConcertSessionActivity> InActivity, ESelectInfo::Type SelectInfo)
 {
 	UpdateDetailArea(InActivity);
 }
@@ -593,7 +592,7 @@ void SConcertSessionActivities::OnDetailsAreaExpansionChanged(bool bExpanded)
 	UpdateDetailArea(bDetailsAreaExpanded ? GetSelectedActivity() : nullptr);
 }
 
-void SConcertSessionActivities::UpdateDetailArea(TSharedPtr<FConcertClientSessionActivity> InSelectedActivity)
+void SConcertSessionActivities::UpdateDetailArea(TSharedPtr<FConcertSessionActivity> InSelectedActivity)
 {
 	if (DetailsAreaVisibility != EVisibility::Visible || !bDetailsAreaExpanded)
 	{
@@ -705,13 +704,13 @@ EConcertActivityFilterFlags SConcertSessionActivities::QueryActiveActivityFilter
 void SConcertSessionActivities::OnActivityFilterUpdated()
 {
 	// Try preserving the selected activity.
-	TSharedPtr<FConcertClientSessionActivity> SelectedActivity = GetSelectedActivity();
+	TSharedPtr<FConcertSessionActivity> SelectedActivity = GetSelectedActivity();
 
 	// Reset the list of displayed activities.
 	Activities.Reset(AllActivities.Num());
 
 	// Apply the filter.
-	for (TSharedPtr<FConcertClientSessionActivity>& Activity : AllActivities)
+	for (TSharedPtr<FConcertSessionActivity>& Activity : AllActivities)
 	{
 		if (PassesFilters(*Activity))
 		{
@@ -797,7 +796,7 @@ void SConcertSessionActivities::FetchActivities()
 	}
 }
 
-void SConcertSessionActivities::Append(TSharedPtr<FConcertClientSessionActivity> Activity)
+void SConcertSessionActivities::Append(TSharedPtr<FConcertSessionActivity> Activity)
 {
 	if (Activity->Activity.bIgnored)
 	{
@@ -834,7 +833,7 @@ void SConcertSessionActivities::Reset()
 	IgnoredActivityNum = 0;
 }
 
-bool SConcertSessionActivities::PassesFilters(const FConcertClientSessionActivity& Activity)
+bool SConcertSessionActivities::PassesFilters(const FConcertSessionActivity& Activity)
 {
 	if (Activity.Activity.EventType == EConcertSyncActivityEventType::Connection && ConnectionActivitiesVisibility.Get() != EVisibility::Visible) // Filter out 'connection' activities?
 	{
@@ -866,7 +865,7 @@ FText SConcertSessionActivities::UpdateTextFilter(const FText& InFilterText)
 	return SearchTextFilter->GetFilterErrorText();
 }
 
-void SConcertSessionActivities::PopulateSearchStrings(const FConcertClientSessionActivity& Activity, TArray<FString>& OutSearchStrings) const
+void SConcertSessionActivities::PopulateSearchStrings(const FConcertSessionActivity& Activity, TArray<FString>& OutSearchStrings) const
 {
 	FText ClientName = GetActivityUserFn ? ConcertSessionActivityUtils::GetClientName(GetActivityUserFn(Activity.Activity.EndpointId)) : FText::GetEmpty();
 
@@ -889,7 +888,7 @@ void SConcertSessionActivities::PopulateSearchStrings(const FConcertClientSessio
 	}
 }
 
-TSharedRef<ITableRow> SConcertSessionActivities::OnGenerateActivityRowWidget(TSharedPtr<FConcertClientSessionActivity> Activity, const TSharedRef<STableViewBase>& OwnerTable)
+TSharedRef<ITableRow> SConcertSessionActivities::OnGenerateActivityRowWidget(TSharedPtr<FConcertSessionActivity> Activity, const TSharedRef<STableViewBase>& OwnerTable)
 {
 	return SNew(SConcertSessionActivityRow, Activity, GetActivityUserFn ? GetActivityUserFn(Activity->Activity.EndpointId) : nullptr, OwnerTable)
 		.TimeFormat(TimeFormat)
@@ -897,13 +896,13 @@ TSharedRef<ITableRow> SConcertSessionActivities::OnGenerateActivityRowWidget(TSh
 		.OnMakeColumnOverlayWidget(MakeColumnOverlayWidgetFn);
 }
 
-TSharedPtr<FConcertClientSessionActivity> SConcertSessionActivities::GetSelectedActivity() const
+TSharedPtr<FConcertSessionActivity> SConcertSessionActivities::GetSelectedActivity() const
 {
-	TArray<TSharedPtr<FConcertClientSessionActivity>> SelectedItems = ActivityView->GetSelectedItems();
+	TArray<TSharedPtr<FConcertSessionActivity>> SelectedItems = ActivityView->GetSelectedItems();
 	return SelectedItems.Num() ? SelectedItems[0] : nullptr;
 }
 
-TSharedPtr<FConcertClientSessionActivity> SConcertSessionActivities::GetMostRecentActivity() const
+TSharedPtr<FConcertSessionActivity> SConcertSessionActivities::GetMostRecentActivity() const
 {
 	// NOTE: This function assumes that activities are sorted by ID. When used for recovery purpose,
 	//       the activities are listed from the most recent to the oldest. When displaying a live
@@ -925,7 +924,7 @@ bool SConcertSessionActivities::IsLastColumn(const FName& ColumnId) const
 	return ColumnId == ConcertSessionActivityUtils::SummaryColumnId; // Summary column is always visible and always the last.
 }
 
-void SConcertSessionActivities::DisplayTransactionDetails(const FConcertClientSessionActivity& Activity, const FConcertTransactionEventBase& InTransaction)
+void SConcertSessionActivities::DisplayTransactionDetails(const FConcertSessionActivity& Activity, const FConcertTransactionEventBase& InTransaction)
 {
 	const FConcertSyncTransactionActivitySummary* Summary = Activity.ActivitySummary.Cast<FConcertSyncTransactionActivitySummary>();
 	FString TransactionTitle = Summary ? Summary->TransactionTitle.ToString() : FString();
@@ -958,7 +957,7 @@ void SConcertSessionActivities::DisplayTransactionDetails(const FConcertClientSe
 	SetDetailsPanelVisibility(TransactionDetailsPanel.Get());
 }
 
-void SConcertSessionActivities::DisplayPackageDetails(const FConcertClientSessionActivity& Activity, int64 PackageRevision, const FConcertPackageInfo& PackageInfo)
+void SConcertSessionActivities::DisplayPackageDetails(const FConcertSessionActivity& Activity, int64 PackageRevision, const FConcertPackageInfo& PackageInfo)
 {
 	const FConcertClientInfo* ClientInfo = nullptr;
 	if (GetActivityUserFn)
diff --git a/Engine/Plugins/Developer/Concert/ConcertUI/ConcertSharedSlate/Source/ConcertSharedSlate/Private/SConcertSessionRecovery.cpp b/Engine/Plugins/Developer/Concert/ConcertUI/ConcertSharedSlate/Source/ConcertSharedSlate/Private/SConcertSessionRecovery.cpp
--- a/Engine/Plugins/Developer/Concert/ConcertUI/ConcertSharedSlate/Source/ConcertSharedSlate/Private/SConcertSessionRecovery.cpp
+++ b/Engine/Plugins/Developer/Concert/ConcertUI/ConcertSharedSlate/Source/ConcertSharedSlate/Private/SConcertSessionRecovery.cpp
@@ -2,7 +2,6 @@
 
 #include "SConcertSessionRecovery.h"
 
-#include "IConcertClientWorkspace.h"
 #include "EditorStyleSet.h"
 #include "EditorFontGlyphs.h"
 #include "Widgets/SWindow.h"
@@ -35,7 +34,7 @@ void SConcertSessionRecovery::Construct(const FArguments& InArgs)
 	SAssignNew(ActivityView, SConcertSessionActivities)
 	.OnFetchActivities(InArgs._OnFetchActivities)
 	.OnMapActivityToClient(InArgs._OnMapActivityToClient)
-	.OnMakeColumnOverlayWidget([this](TWeakPtr<FConcertClientSessionActivity> Activity, const FName& ColumnId) { return MakeRecoverThroughWidget(Activity, ColumnId); })
+	.OnMakeColumnOverlayWidget([this](TWeakPtr<FConcertSessionActivity> Activity, const FName& ColumnId) { return MakeRecoverThroughWidget(Activity, ColumnId); })
 	.HighlightText(this, &SConcertSessionRecovery::HighlightSearchText)
 	.TimeFormat(ActivityViewOptions.Get(), &FConcertSessionActivitiesOptions::GetTimeFormat)
 	.ClientAvatarColorColumnVisibility(InArgs._ClientAvatarColorColumnVisibility)
@@ -162,9 +161,9 @@ void SConcertSessionRecovery::Construct(const FArguments& InArgs)
 	];
 }
 
-TSharedPtr<SWidget> SConcertSessionRecovery::MakeRecoverThroughWidget(TWeakPtr<FConcertClientSessionActivity> Activity, const FName& ColumnId)
+TSharedPtr<SWidget> SConcertSessionRecovery::MakeRecoverThroughWidget(TWeakPtr<FConcertSessionActivity> Activity, const FName& ColumnId)
 {
-	if (TSharedPtr<FConcertClientSessionActivity> ActivityPin = Activity.Pin())
+	if (TSharedPtr<FConcertSessionActivity> ActivityPin = Activity.Pin())
 	{
 		if (ActivityView->IsLastColumn(ColumnId) && !ActivityPin->Activity.bIgnored) // The most right cell and activity will not be ignored on recovery.
 		{
@@ -203,12 +202,12 @@ int32 SConcertSessionRecovery::GetTotalActivityNum() const
 	return ActivityView->GetTotalActivityNum();
 }
 
-TSharedPtr<FConcertClientSessionActivity> SConcertSessionRecovery::GetMostRecentActivity() const
+TSharedPtr<FConcertSessionActivity> SConcertSessionRecovery::GetMostRecentActivity() const
 {
 	return ActivityView->GetMostRecentActivity();
 }
 
-EVisibility SConcertSessionRecovery::GetRecoverThroughButtonVisibility(TSharedPtr<FConcertClientSessionActivity> Activity)
+EVisibility SConcertSessionRecovery::GetRecoverThroughButtonVisibility(TSharedPtr<FConcertSessionActivity> Activity)
 {
 	return Activity == ActivityView->GetSelectedActivity() && IsRecoverThroughButtonVisible.Get() ? EVisibility::Visible : EVisibility::Hidden;
 }
@@ -292,7 +291,7 @@ FReply SConcertSessionRecovery::OnRecoverAllClicked()
 	return FReply::Handled();
 }
 
-void SConcertSessionRecovery::RecoverThrough(TSharedPtr<FConcertClientSessionActivity> Item)
+void SConcertSessionRecovery::RecoverThrough(TSharedPtr<FConcertSessionActivity> Item)
 {
 	bool bShouldDismissWindow = true;
 	if (Item)
diff --git a/Engine/Plugins/Developer/Concert/ConcertUI/ConcertSharedSlate/Source/ConcertSharedSlate/Public/ConcertFrontendUtils.h b/Engine/Plugins/Developer/Concert/ConcertUI/ConcertSharedSlate/Source/ConcertSharedSlate/Public/ConcertFrontendUtils.h
--- a/Engine/Plugins/Developer/Concert/ConcertUI/ConcertSharedSlate/Source/ConcertSharedSlate/Public/ConcertFrontendUtils.h
+++ b/Engine/Plugins/Developer/Concert/ConcertUI/ConcertSharedSlate/Source/ConcertSharedSlate/Public/ConcertFrontendUtils.h
@@ -3,38 +3,15 @@
 #pragma once
 
 #include "CoreMinimal.h"
-#include "ConcertActionDefinition.h"
 #include "Styling/AppStyle.h"
 #include "Widgets/Layout/SBorder.h"
-#include "Widgets/Images/SImage.h"
-#include "Widgets/Input/SButton.h"
 #include "Widgets/Text/STextBlock.h"
-#include "Widgets/Layout/SBox.h"
 #include "Widgets/Layout/SExpandableArea.h"
-#include "Widgets/SBoxPanel.h"
-#include "EditorFontGlyphs.h"
-#include "ConcertFrontendStyle.h"
 
 #define LOCTEXT_NAMESPACE "ConcertFrontendUtils"
 
 namespace ConcertFrontendUtils
 {
-	static const FName ButtonIconSyle = TEXT("FontAwesome.10");
-	static const float MinDesiredWidthForBtnAndIcon = 29.f;
-	static const FName ButtonStyleNames[(int32)EConcertActionType::NUM] = {
-		TEXT("SimpleButton"),
-		TEXT("FlatButton.Primary"),
-		TEXT("FlatButton.Info"),
-		TEXT("FlatButton.Success"),
-		TEXT("FlatButton.Warning"),
-		TEXT("FlatButton.Danger"),
-	};
-
-	FORCEINLINE bool ShowSessionConnectionUI()
-	{
-		return !(IS_PROGRAM);
-	}
-
 	inline TSharedRef<SWidget> CreateDisplayName(const TAttribute<FText>& InDisplayName)
 	{
 		return SNew(SBorder)
@@ -47,66 +24,7 @@ namespace ConcertFrontendUtils
 				.Text(InDisplayName)
 			];
 	}
-
-	inline TSharedRef<SButton> CreateTextButton(const FConcertActionDefinition& InDef)
-	{
-		const FButtonStyle* ButtonStyle = &FAppStyle::Get().GetWidgetStyle<FButtonStyle>(ButtonStyleNames[(int32)InDef.Type]);
-		check(ButtonStyle);
-		const float ButtonContentWidthPadding = 6.f;
-		const float PaddingCompensation = (ButtonStyle->NormalPadding.Left + ButtonStyle->NormalPadding.Right + ButtonContentWidthPadding * 2);
-
-		return SNew(SButton)
-			.ButtonStyle(ButtonStyle)
-			.ToolTipText(InDef.ToolTipText)
-			.ContentPadding(FMargin(ButtonContentWidthPadding, 2.f))
-			.IsEnabled(InDef.IsEnabled)
-			.Visibility_Lambda([IsVisible = InDef.IsVisible]() { return IsVisible.Get() ? EVisibility::Visible : EVisibility::Collapsed; })
-			.OnClicked_Lambda([OnExecute = InDef.OnExecute]() { OnExecute.ExecuteIfBound(); return FReply::Handled(); })
-			[
-				SNew(SBox)
-				.MinDesiredWidth(MinDesiredWidthForBtnAndIcon - PaddingCompensation)
-				.HAlign(HAlign_Center)
-				[
-					SNew(STextBlock)
-					.Font(FAppStyle::Get().GetFontStyle(ButtonIconSyle))
-					.Text(InDef.Text)
-					.Justification(ETextJustify::Center)
-				]
-			];
-	}
-
-	inline TSharedRef<SButton> CreateIconButton(const FConcertActionDefinition& InDef)
-	{
-		return SNew(SButton)
-			.ButtonStyle(FAppStyle::Get(), ButtonStyleNames[(int32)InDef.Type])
-			.ToolTipText(InDef.ToolTipText)
-			.ContentPadding(FMargin(0, 0))
-			.IsEnabled(InDef.IsEnabled)
-			.Visibility_Lambda([IsVisible = InDef.IsVisible]() { return IsVisible.Get() ? EVisibility::Visible : EVisibility::Collapsed; })
-			.OnClicked_Lambda([OnExecute = InDef.OnExecute]() { OnExecute.ExecuteIfBound(); return FReply::Handled(); })
-			.VAlign(VAlign_Center)
-			.HAlign(HAlign_Center)
-			[
-				SNew(SImage)
-				.Image(TAttribute<const FSlateBrush*>::Create([IconStyleAttr = InDef.IconStyle]() { return FConcertFrontendStyle::Get()->GetBrush(IconStyleAttr.Get()); }))
-				.ColorAndOpacity(FSlateColor::UseForeground())
-			];
-	}
-
-	inline void AppendButtons(TSharedRef<SHorizontalBox> InHorizBox, TArrayView<const FConcertActionDefinition> InDefs)
-	{
-		for (const FConcertActionDefinition& Def : InDefs)
-		{
-			InHorizBox->AddSlot()
-				.AutoWidth()
-				.VAlign(VAlign_Center)
-				.Padding(FMargin(1.0f))
-				[
-					Def.IconStyle.IsSet() ? CreateIconButton(Def) : CreateTextButton(Def)
-				];
-		}
-	}
-
+	
 	template <typename ItemType, typename PredFactoryType>
 	inline void SyncArraysByPredicate(TArray<TSharedPtr<ItemType>>& InOutArray, TArray<TSharedPtr<ItemType>>&& InNewArray, const PredFactoryType& InPredFactory)
 	{
diff --git a/Engine/Plugins/Developer/Concert/ConcertUI/ConcertSharedSlate/Source/ConcertSharedSlate/Public/SConcertSessionActivities.h b/Engine/Plugins/Developer/Concert/ConcertUI/ConcertSharedSlate/Source/ConcertSharedSlate/Public/SConcertSessionActivities.h
--- a/Engine/Plugins/Developer/Concert/ConcertUI/ConcertSharedSlate/Source/ConcertSharedSlate/Public/SConcertSessionActivities.h
+++ b/Engine/Plugins/Developer/Concert/ConcertUI/ConcertSharedSlate/Source/ConcertSharedSlate/Public/SConcertSessionActivities.h
@@ -11,7 +11,7 @@
 #include "Widgets/Views/SListView.h"
 #include "ConcertSyncSessionTypes.h"
 
-struct FConcertClientSessionActivity;
+struct FConcertSessionActivity;
 struct FConcertClientInfo;
 struct FConcertPackageInfo;
 struct FConcertTransactionEventBase;
@@ -45,19 +45,19 @@ public:
 	};
 
 	/** Used to pull activities from a session. Used to fetch and display the activities of an archived session. */
-	using FFetchActivitiesFunc = TFunction<bool(TArray<TSharedPtr<FConcertClientSessionActivity>>& /*InOutActivities*/, int32& /*OutFetchedCount*/, FText& /*ErrorMsg*/)>;
+	using FFetchActivitiesFunc = TFunction<bool(TArray<TSharedPtr<FConcertSessionActivity>>& /*InOutActivities*/, int32& /*OutFetchedCount*/, FText& /*ErrorMsg*/)>;
 
 	/** Used to map an activity to its client. */
 	using FGetActivityClientInfoFunc = TFunction<const FConcertClientInfo*(FGuid /*ClientId*/)>;
 
 	/** Returns the transaction event corresponding the specified activity.*/
-	using FGetTransactionEvent = TFunction<TFuture<TOptional<FConcertSyncTransactionEvent>>(const FConcertClientSessionActivity& /*Activity*/)>;
+	using FGetTransactionEvent = TFunction<TFuture<TOptional<FConcertSyncTransactionEvent>>(const FConcertSessionActivity& /*Activity*/)>;
 
 	/** Returns the package event corresponding to the package activity. */
-	using FGetPackageEvent = TFunction<bool(const FConcertClientSessionActivity& /*Activity*/, FConcertSyncPackageEventMetaData& /*OutEvent*/)>;
+	using FGetPackageEvent = TFunction<bool(const FConcertSessionActivity& /*Activity*/, FConcertSyncPackageEventMetaData& /*OutEvent*/)>;
 
 	/** Used to overlay a widget over a column widget to add custom functionalities to a row. */
-	using FMakeColumnOverlayWidgetFunc = TFunction<TSharedPtr<SWidget>(TWeakPtr<FConcertClientSessionActivity> /*ThisRowActivity*/, const FName& /*ColumnId*/)>;
+	using FMakeColumnOverlayWidgetFunc = TFunction<TSharedPtr<SWidget>(TWeakPtr<FConcertSessionActivity> /*ThisRowActivity*/, const FName& /*ColumnId*/)>;
 
 public:
 	SLATE_BEGIN_ARGS(SConcertSessionActivities)
@@ -145,7 +145,7 @@ public:
 	virtual void Tick(const FGeometry& AllottedGeometry, const double InCurrentTime, const float InDeltaTime) override;
 
 	/** Returns the activity selected or null if none is selected. */
-	TSharedPtr<FConcertClientSessionActivity> GetSelectedActivity() const;
+	TSharedPtr<FConcertSessionActivity> GetSelectedActivity() const;
 
 	/** Returns the total number of activities currently stored (no filter applied). */
 	int32 GetTotalActivityNum() const { return AllActivities.Num(); }
@@ -157,7 +157,7 @@ public:
 	int32 GetIgnoredActivityNum() const { return IgnoredActivityNum; }
 
 	/** Returns the most recent activity available, ignoring the current filter. */
-	TSharedPtr<FConcertClientSessionActivity> GetMostRecentActivity() const;
+	TSharedPtr<FConcertSessionActivity> GetMostRecentActivity() const;
 
 	/** Returns true if the column names is the last one (most right one). */
 	bool IsLastColumn(const FName& ColumnId) const;
@@ -166,7 +166,7 @@ public:
 	void Reset();
 
 	/** Append an activity to the view. Used to populate the view from a live session. */
-	void Append(TSharedPtr<FConcertClientSessionActivity> Activity);
+	void Append(TSharedPtr<FConcertSessionActivity> Activity);
 
 	/** Request the view to refresh. */
 	void RequestRefresh();
@@ -176,7 +176,7 @@ public:
 
 private:
 	/** Generates a row widget in the table view. */
-	TSharedRef<ITableRow> OnGenerateActivityRowWidget(TSharedPtr<FConcertClientSessionActivity> Item, const TSharedRef<STableViewBase>& OwnerTable);
+	TSharedRef<ITableRow> OnGenerateActivityRowWidget(TSharedPtr<FConcertSessionActivity> Item, const TSharedRef<STableViewBase>& OwnerTable);
 
 	/** Fetches more activities. */
 	void FetchActivities();
@@ -188,16 +188,16 @@ private:
 	void OnActivityFilterUpdated();
 
 	/** Returns true if the specified activity passes the view filters. */
-	bool PassesFilters(const FConcertClientSessionActivity& Activity);
+	bool PassesFilters(const FConcertSessionActivity& Activity);
 
 	/** Invoked when the list view scrolled to fetch more data from the activity provider. */
 	void OnListViewScrolled(double InScrollOffset);
 
 	/** Invoked when the list view selection changes. */
-	void OnListViewSelectionChanged(TSharedPtr<FConcertClientSessionActivity> InActivity, ESelectInfo::Type SelectInfo);
+	void OnListViewSelectionChanged(TSharedPtr<FConcertSessionActivity> InActivity, ESelectInfo::Type SelectInfo);
 
 	/** Convert the item into a list of text element used search bars. */
-	void PopulateSearchStrings(const FConcertClientSessionActivity& Item, TArray<FString>& OutSearchStrings) const;
+	void PopulateSearchStrings(const FConcertSessionActivity& Item, TArray<FString>& OutSearchStrings) const;
 
 	/** Returns whether the details expandable area should be visible or not. */
 	EVisibility GetDetailAreaVisibility() const { return DetailsAreaVisibility; }
@@ -209,13 +209,13 @@ private:
 	void OnDetailsAreaExpansionChanged(bool bExpanded);
 
 	/** Displays what properties/objects changed during a transaction in the details area when a transaction activity is selected in the list. */
-	void DisplayTransactionDetails(const FConcertClientSessionActivity& Activity, const FConcertTransactionEventBase& InTransaction);
+	void DisplayTransactionDetails(const FConcertSessionActivity& Activity, const FConcertTransactionEventBase& InTransaction);
 
 	/** Displays what changed in a package in the details area when a package activity is selected in the list. */
-	void DisplayPackageDetails(const FConcertClientSessionActivity& Activity, int64 PackageRevision, const FConcertPackageInfo& PackageInfo);
+	void DisplayPackageDetails(const FConcertSessionActivity& Activity, int64 PackageRevision, const FConcertPackageInfo& PackageInfo);
 
 	/** Update the detail area to display the proper panel and details. */
-	void UpdateDetailArea(TSharedPtr<FConcertClientSessionActivity> InSelectedActivity);
+	void UpdateDetailArea(TSharedPtr<FConcertSessionActivity> InSelectedActivity);
 
 	/** Change the visibility of the details panels, keeping only one visible. */
 	void SetDetailsPanelVisibility(const SWidget* VisiblePanel);
@@ -228,13 +228,13 @@ private:
 
 private:
 	/** List of all activities (including the filtered out ones). */
-	TArray<TSharedPtr<FConcertClientSessionActivity>> AllActivities;
+	TArray<TSharedPtr<FConcertSessionActivity>> AllActivities;
 
 	/** List of currently displayed activities. */
-	TArray<TSharedPtr<FConcertClientSessionActivity>> Activities;
+	TArray<TSharedPtr<FConcertSessionActivity>> Activities;
 
 	/** The list view widget displaying the activities. */
-	TSharedPtr<SListView<TSharedPtr<FConcertClientSessionActivity>>> ActivityView;
+	TSharedPtr<SListView<TSharedPtr<FConcertSessionActivity>>> ActivityView;
 
 	/** Used to overlay a widget over a column widget (add an extra layer above the normal one) */
 	FMakeColumnOverlayWidgetFunc MakeColumnOverlayWidgetFn;
@@ -303,7 +303,7 @@ private:
 	bool bAllActivitiesFetched = false;
 
 	/** Utility class used to tokenize and match text displayed in the list view. */
-	TSharedPtr<TTextFilter<const FConcertClientSessionActivity&>> SearchTextFilter;
+	TSharedPtr<TTextFilter<const FConcertSessionActivity&>> SearchTextFilter;
 
 	/** The expandable area under which the activity details are displayed. */
 	TSharedPtr<SExpandableArea> ExpandableDetails;
diff --git a/Engine/Plugins/Developer/Concert/ConcertUI/ConcertSharedSlate/Source/ConcertSharedSlate/Public/SConcertSessionRecovery.h b/Engine/Plugins/Developer/Concert/ConcertUI/ConcertSharedSlate/Source/ConcertSharedSlate/Public/SConcertSessionRecovery.h
--- a/Engine/Plugins/Developer/Concert/ConcertUI/ConcertSharedSlate/Source/ConcertSharedSlate/Public/SConcertSessionRecovery.h
+++ b/Engine/Plugins/Developer/Concert/ConcertUI/ConcertSharedSlate/Source/ConcertSharedSlate/Public/SConcertSessionRecovery.h
@@ -7,7 +7,7 @@
 #include "Widgets/DeclarativeSyntaxSupport.h"
 #include "Widgets/SCompoundWidget.h"
 
-struct FConcertClientSessionActivity;
+struct FConcertSessionActivity;
 struct FConcertClientInfo;
 class SConcertSessionActivities;
 class FConcertSessionActivitiesOptions;
@@ -19,7 +19,7 @@ class CONCERTSHAREDSLATE_API SConcertSessionRecovery : public SCompoundWidget
 {
 public:
 	/** Used to pull activities from a session. Used to fetch and display the activities of an archived session. */
-	using FFetchActivitiesFunc = TFunction<bool(TArray<TSharedPtr<FConcertClientSessionActivity>>& /*InOutActivities*/, int32& /*OutFetchedCount*/, FText& /*ErrorMsg*/)>;
+	using FFetchActivitiesFunc = TFunction<bool(TArray<TSharedPtr<FConcertSessionActivity>>& /*InOutActivities*/, int32& /*OutFetchedCount*/, FText& /*ErrorMsg*/)>;
 
 	/** Used to map an activity to its client. */
 	using FGetActivityClientInfoFunc = TFunction<const FConcertClientInfo*(FGuid /*ClientId*/)>;
@@ -52,7 +52,7 @@ public:
 		SLATE_ARGUMENT(FGetActivityClientInfoFunc, OnMapActivityToClient)
 
 		/** Invoked when the user clicks the 'recover' button. */
-		SLATE_ARGUMENT(TFunction<bool(TSharedPtr<FConcertClientSessionActivity>)>, OnRestore)
+		SLATE_ARGUMENT(TFunction<bool(TSharedPtr<FConcertSessionActivity>)>, OnRestore)
 
 		/** Invoked when the user clicks the 'cancel' button. */
 		SLATE_ARGUMENT(TFunction<void()>, OnCancel)
@@ -104,7 +104,7 @@ public:
 	void Construct(const FArguments& InArgs);
 
 	/** Returns the activity, selected by the user, through which the session should be (or was) recovered or null to prevent recovery. */
-	TSharedPtr<FConcertClientSessionActivity> GetRecoverThroughItem() { return RecoveryThroughItem; }
+	TSharedPtr<FConcertSessionActivity> GetRecoverThroughItem() { return RecoveryThroughItem; }
 
 	/** Removes all activities from the view and reset the activity stream. */
 	void Reset();
@@ -116,7 +116,7 @@ public:
 	FText GetRecoverAllButtonTooltip() const;
 
 	/** Returns the most recent activity available, ignoring the current filter. */
-	TSharedPtr<FConcertClientSessionActivity> GetMostRecentActivity() const;
+	TSharedPtr<FConcertSessionActivity> GetMostRecentActivity() const;
 
 private:
 	void OnSearchTextChanged(const FText& InFilterText);
@@ -126,10 +126,10 @@ private:
 	FReply OnCancelRecoveryClicked();
 	FReply OnRecoverAllClicked();
 
-	EVisibility GetRecoverThroughButtonVisibility(TSharedPtr<FConcertClientSessionActivity> Activity);
+	EVisibility GetRecoverThroughButtonVisibility(TSharedPtr<FConcertSessionActivity> Activity);
 	FText GetRecoverThroughButtonTooltip() const;
-	TSharedPtr<SWidget> MakeRecoverThroughWidget(TWeakPtr<FConcertClientSessionActivity>, const FName&);
-	void RecoverThrough(TSharedPtr<FConcertClientSessionActivity> Item);
+	TSharedPtr<SWidget> MakeRecoverThroughWidget(TWeakPtr<FConcertSessionActivity>, const FName&);
+	void RecoverThrough(TSharedPtr<FConcertSessionActivity> Item);
 
 	/** Close the windows hosting this recovery widget. */
 	void DismissWindow();
@@ -142,7 +142,7 @@ private:
 	TSharedPtr<FConcertSessionActivitiesOptions> ActivityViewOptions;
 
 	/** The activity selected when the user click 'Recover' or 'Recover Through' buttons. */
-	TSharedPtr<FConcertClientSessionActivity> RecoveryThroughItem;
+	TSharedPtr<FConcertSessionActivity> RecoveryThroughItem;
 
 	/** The parent window hosting this widget. */
 	TWeakPtr<SWindow> ParentWindow;
@@ -160,7 +160,7 @@ private:
 	TAttribute<bool> IsRecoverThroughButtonVisible;
 
 	/** The function invoked when the user clicks the restore button. Might not be bound. */
-	TFunction<bool(TSharedPtr<FConcertClientSessionActivity>)> OnRestoreFn;
+	TFunction<bool(TSharedPtr<FConcertSessionActivity>)> OnRestoreFn;
 
 	/** The function invoked when the user clicks the cancel button. Might not be bound. */
 	TFunction<void()> OnCancelFn;
diff --git a/Engine/Plugins/Developer/Concert/DisasterRecoveryClient/Source/DisasterRecoveryClient/Private/DisasterRecoverySessionManager.cpp b/Engine/Plugins/Developer/Concert/DisasterRecoveryClient/Source/DisasterRecoveryClient/Private/DisasterRecoverySessionManager.cpp
--- a/Engine/Plugins/Developer/Concert/DisasterRecoveryClient/Source/DisasterRecoveryClient/Private/DisasterRecoverySessionManager.cpp
+++ b/Engine/Plugins/Developer/Concert/DisasterRecoveryClient/Source/DisasterRecoveryClient/Private/DisasterRecoverySessionManager.cpp
@@ -810,7 +810,7 @@ TFuture<TPair<bool, FText>> FDisasterRecoverySessionManager::CreateAndJoinSessio
 	return Promise->GetFuture();
 }
 
-TFuture<TPair<bool, FText>> FDisasterRecoverySessionManager::RestoreAndJoinSession(TSharedPtr<FDisasterRecoverySession> Session, TSharedPtr<FConcertClientSessionActivity> ThroughActivity)
+TFuture<TPair<bool, FText>> FDisasterRecoverySessionManager::RestoreAndJoinSession(TSharedPtr<FDisasterRecoverySession> Session, TSharedPtr<FConcertSessionActivity> ThroughActivity)
 {
 	// Failed if another session is already in progress or active.
 	if (HasInProgressSession())
diff --git a/Engine/Plugins/Developer/Concert/DisasterRecoveryClient/Source/DisasterRecoveryClient/Private/DisasterRecoverySessionManager.h b/Engine/Plugins/Developer/Concert/DisasterRecoveryClient/Source/DisasterRecoveryClient/Private/DisasterRecoverySessionManager.h
--- a/Engine/Plugins/Developer/Concert/DisasterRecoveryClient/Source/DisasterRecoveryClient/Private/DisasterRecoverySessionManager.h
+++ b/Engine/Plugins/Developer/Concert/DisasterRecoveryClient/Source/DisasterRecoveryClient/Private/DisasterRecoverySessionManager.h
@@ -20,7 +20,7 @@ DECLARE_MULTICAST_DELEGATE_OneParam(FOnDisasterRecoverySessionUpdated, TSharedRe
 
 class IConcertClientSession;
 class FConcertActivityStream;
-struct FConcertClientSessionActivity;
+struct FConcertSessionActivity;
 
 /**
  * Monitors, creates, restores, imports and inspects disaster recovery sessions.
@@ -101,7 +101,7 @@ public:
 	 * @return True if the session was discarded, false and an error message if it failed.
 	 * @note OnSessionAdded() is invoked if the session is created and the original archived session remains unmodified.
 	 */
-	TFuture<TPair<bool, FText>> RestoreAndJoinSession(TSharedPtr<FDisasterRecoverySession> ArchivedSession, TSharedPtr<FConcertClientSessionActivity> ThroughActivity);
+	TFuture<TPair<bool, FText>> RestoreAndJoinSession(TSharedPtr<FDisasterRecoverySession> ArchivedSession, TSharedPtr<FConcertSessionActivity> ThroughActivity);
 
 	/** Leave the current live session, if any. If is safe to call the function even if the user doesn't have a live session in progress. */
 	void LeaveSession();
diff --git a/Engine/Plugins/Developer/Concert/DisasterRecoveryClient/Source/DisasterRecoveryClient/Private/DisasterRecoveryTasks.cpp b/Engine/Plugins/Developer/Concert/DisasterRecoveryClient/Source/DisasterRecoveryClient/Private/DisasterRecoveryTasks.cpp
--- a/Engine/Plugins/Developer/Concert/DisasterRecoveryClient/Source/DisasterRecoveryClient/Private/DisasterRecoveryTasks.cpp
+++ b/Engine/Plugins/Developer/Concert/DisasterRecoveryClient/Source/DisasterRecoveryClient/Private/DisasterRecoveryTasks.cpp
@@ -432,7 +432,7 @@ bool FFindRecoverableDisasterRecoverySessionTask::Tick()
 
 		if (ReadCount > 0) // On error, ReadCount is 0.
 		{
-			for (const TSharedPtr<FConcertClientSessionActivity>& Activity : Activities)
+			for (const TSharedPtr<FConcertSessionActivity>& Activity : Activities)
 			{
 				if (!Activity->Activity.bIgnored)
 				{
diff --git a/Engine/Plugins/Developer/Concert/DisasterRecoveryClient/Source/DisasterRecoveryClient/Private/DisasterRecoveryTasks.h b/Engine/Plugins/Developer/Concert/DisasterRecoveryClient/Source/DisasterRecoveryClient/Private/DisasterRecoveryTasks.h
--- a/Engine/Plugins/Developer/Concert/DisasterRecoveryClient/Source/DisasterRecoveryClient/Private/DisasterRecoveryTasks.h
+++ b/Engine/Plugins/Developer/Concert/DisasterRecoveryClient/Source/DisasterRecoveryClient/Private/DisasterRecoveryTasks.h
@@ -10,7 +10,7 @@
 class IConcertSyncClient;
 class IConcertClientSession;
 class FConcertActivityStream;
-struct FConcertClientSessionActivity;
+struct FConcertSessionActivity;
 struct FDisasterRecoverySession;
 
 
@@ -199,7 +199,7 @@ private:
 	TFunction<void(TSharedPtr<FDisasterRecoverySession>)> OnCandidateSearchResultFn;
 	TFunction<void(const FText&)> OnErrorFn;
 	TSharedPtr<FConcertActivityStream> ActivityStream;
-	TArray<TSharedPtr<FConcertClientSessionActivity>> Activities;
+	TArray<TSharedPtr<FConcertSessionActivity>> Activities;
 	TUniquePtr<FMountDisasterRecoverySessionRepositoryTask> MountingTask;
 	TUniquePtr<FLookupDisasterRecoverySessionIdTask> SearchingSessionIdTask;
 	FGuid SessionId;
diff --git a/Engine/Plugins/Developer/Concert/DisasterRecoveryClient/Source/DisasterRecoveryClient/Private/SDisasterRecoveryHub.cpp b/Engine/Plugins/Developer/Concert/DisasterRecoveryClient/Source/DisasterRecoveryClient/Private/SDisasterRecoveryHub.cpp
--- a/Engine/Plugins/Developer/Concert/DisasterRecoveryClient/Source/DisasterRecoveryClient/Private/SDisasterRecoveryHub.cpp
+++ b/Engine/Plugins/Developer/Concert/DisasterRecoveryClient/Source/DisasterRecoveryClient/Private/SDisasterRecoveryHub.cpp
@@ -596,7 +596,7 @@ void SDisasterRecoveryHub::OnSessionSelectionChanged(TSharedPtr<FDisasterRecover
 TSharedRef<SWidget> SDisasterRecoveryHub::MakeSessionActivityView()
 {
 	return SAssignNew(ActivityView, SConcertSessionRecovery)
-		.OnFetchActivities([this](TArray<TSharedPtr<FConcertClientSessionActivity>>& InOutActivities, int32& OutFetchedCount, FText& ErrorMsg) { return FetchActivities(InOutActivities, OutFetchedCount, ErrorMsg); })
+		.OnFetchActivities([this](TArray<TSharedPtr<FConcertSessionActivity>>& InOutActivities, int32& OutFetchedCount, FText& ErrorMsg) { return FetchActivities(InOutActivities, OutFetchedCount, ErrorMsg); })
 		.ClientAvatarColorColumnVisibility(EVisibility::Hidden) // Disaster recovery has only one user, the local one.
 		.ClientNameColumnVisibility(EVisibility::Hidden)
 		.OperationColumnVisibility(EVisibility::Visible)
@@ -609,7 +609,7 @@ TSharedRef<SWidget> SDisasterRecoveryHub::MakeSessionActivityView()
 		.IsIgnoredActivityFilteringEnabled(true)     // Events ignored when restoring are not displayed by default. Enabled to inspect Multi-User transaction/package activities recorded (but not recoverable) by disaster recovery session in case the crash occurred during the Multi-User session.
 		.AreRecoverAllAndCancelButtonsVisible(false) // Replaced by this widget own (for better placement)
 		.IsRecoverThroughButtonsVisible(this, &SDisasterRecoveryHub::IsRecoverThroughButtonVisible)
-		.OnRestore([this](TSharedPtr<FConcertClientSessionActivity> ThroughActivity) { OnRecoverThroughButtonClicked(ThroughActivity); return /*DismissWindow*/false; }) // Invoked when the 'Recover Through' button is clicked.
+		.OnRestore([this](TSharedPtr<FConcertSessionActivity> ThroughActivity) { OnRecoverThroughButtonClicked(ThroughActivity); return /*DismissWindow*/false; }) // Invoked when the 'Recover Through' button is clicked.
 		.NoActivitiesReasonText(this, &SDisasterRecoveryHub::GetNoActivityDisplayedReason);
 }
 
@@ -637,7 +637,7 @@ FText SDisasterRecoveryHub::GetNoActivityDisplayedReason() const
 	}
 }
 
-bool SDisasterRecoveryHub::FetchActivities(TArray<TSharedPtr<FConcertClientSessionActivity>>& InOutActivities, int32& OutFetchedCount, FText& ErrorMsg)
+bool SDisasterRecoveryHub::FetchActivities(TArray<TSharedPtr<FConcertSessionActivity>>& InOutActivities, int32& OutFetchedCount, FText& ErrorMsg)
 {
 	if (SelectedSessionActivityStream)
 	{
@@ -677,12 +677,12 @@ bool SDisasterRecoveryHub::IsRecoverThroughButtonVisible() const
 	return IsRecoverAllButtonEnabled();
 }
 
-void SDisasterRecoveryHub::OnRecoverThroughButtonClicked(TSharedPtr<FConcertClientSessionActivity> ThroughActivity)
+void SDisasterRecoveryHub::OnRecoverThroughButtonClicked(TSharedPtr<FConcertSessionActivity> ThroughActivity)
 {
 	RecoverThrough(ThroughActivity);
 }
 
-FReply SDisasterRecoveryHub::RecoverThrough(TSharedPtr<FConcertClientSessionActivity> ThroughActivity)
+FReply SDisasterRecoveryHub::RecoverThrough(TSharedPtr<FConcertSessionActivity> ThroughActivity)
 {
 	if (SelectedSessionNode && !SelectedSessionNode->IsCategoryNode() && ThroughActivity)
 	{
diff --git a/Engine/Plugins/Developer/Concert/DisasterRecoveryClient/Source/DisasterRecoveryClient/Private/SDisasterRecoveryHub.h b/Engine/Plugins/Developer/Concert/DisasterRecoveryClient/Source/DisasterRecoveryClient/Private/SDisasterRecoveryHub.h
--- a/Engine/Plugins/Developer/Concert/DisasterRecoveryClient/Source/DisasterRecoveryClient/Private/SDisasterRecoveryHub.h
+++ b/Engine/Plugins/Developer/Concert/DisasterRecoveryClient/Source/DisasterRecoveryClient/Private/SDisasterRecoveryHub.h
@@ -12,7 +12,7 @@
 class FConcertActivityStream;
 class FDisasterRecoverySessionManager;
 class SConcertSessionRecovery;
-struct FConcertClientSessionActivity;
+struct FConcertSessionActivity;
 class FDisasterRecoverySessionTreeNode;
 class SWindow;
 
@@ -66,14 +66,14 @@ private:
 	// Selected session activities.
 	TSharedRef<SWidget> MakeSessionActivityView();
 	bool ShouldDisplayActivityDetails() const { return true; }
-	bool FetchActivities(TArray<TSharedPtr<FConcertClientSessionActivity>>& InOutActivities, int32& OutFetchedCount, FText& ErrorMsg);
+	bool FetchActivities(TArray<TSharedPtr<FConcertSessionActivity>>& InOutActivities, int32& OutFetchedCount, FText& ErrorMsg);
 	FText GetNoActivityDisplayedReason() const;
 
 	// Recover buttons.
 	FText GetRecoverAllButtonTooltip() const;
 	FReply OnRecoverAllButtonClicked();
-	void OnRecoverThroughButtonClicked(TSharedPtr<FConcertClientSessionActivity> ThroughActivity);
-	FReply RecoverThrough(TSharedPtr<FConcertClientSessionActivity> ThroughActivity);
+	void OnRecoverThroughButtonClicked(TSharedPtr<FConcertSessionActivity> ThroughActivity);
+	FReply RecoverThrough(TSharedPtr<FConcertSessionActivity> ThroughActivity);
 	bool IsRecoverAllButtonEnabled() const;
 	bool IsRecoverThroughButtonVisible() const;
 
diff --git a/Engine/Plugins/Developer/Concert/MultiUserClient/Source/MultiUserClient/MultiUserClient.Build.cs b/Engine/Plugins/Developer/Concert/MultiUserClient/Source/MultiUserClient/MultiUserClient.Build.cs
--- a/Engine/Plugins/Developer/Concert/MultiUserClient/Source/MultiUserClient/MultiUserClient.Build.cs
+++ b/Engine/Plugins/Developer/Concert/MultiUserClient/Source/MultiUserClient/MultiUserClient.Build.cs
@@ -19,9 +19,10 @@ namespace UnrealBuildTool.Rules
 				{
 					"ApplicationCore",
 					"Concert",
+					"ConcertClientSharedSlate",
+					"ConcertSharedSlate",
 					"ConcertSyncClient",
 					"ConcertSyncCore",
-					"ConcertSharedSlate",
 					"ConcertTransport",
 					"DesktopPlatform",
 					"EditorStyle",
diff --git a/Engine/Plugins/Developer/Concert/MultiUserClient/Source/MultiUserClient/Private/Widgets/SActiveSession.cpp b/Engine/Plugins/Developer/Concert/MultiUserClient/Source/MultiUserClient/Private/Widgets/SActiveSession.cpp
--- a/Engine/Plugins/Developer/Concert/MultiUserClient/Source/MultiUserClient/Private/Widgets/SActiveSession.cpp
+++ b/Engine/Plugins/Developer/Concert/MultiUserClient/Source/MultiUserClient/Private/Widgets/SActiveSession.cpp
@@ -22,6 +22,7 @@
 #include "ConcertFrontendUtils.h"
 #include "ConcertMessageData.h"
 #include "SSessionHistory.h"
+#include "ConcertClientFrontendUtils.h"
 
 #define LOCTEXT_NAMESPACE "SActiveSession"
 
@@ -95,7 +96,7 @@ public:
 			if (ClientInfoPin.IsValid() && SyncClientPin.IsValid())
 			{
 				SyncClientPin->GetSessionClientActions(*ClientInfoPin, ActionDefs);
-				ConcertFrontendUtils::AppendButtons(PresenceCell, ActionDefs);
+				ConcertClientFrontendUtils::AppendButtons(PresenceCell, ActionDefs);
 			}
 			return PresenceCell;
 		}
@@ -288,7 +289,7 @@ public:
 	FSlateFontInfo GetAvatarFont() const
 	{
 		// This font is used to render a small square box filled with the avatar color.
-		FSlateFontInfo ClientIconFontInfo = FAppStyle::Get().GetFontStyle(ConcertFrontendUtils::ButtonIconSyle);
+		FSlateFontInfo ClientIconFontInfo = FAppStyle::Get().GetFontStyle(ConcertClientFrontendUtils::ButtonIconSyle);
 		ClientIconFontInfo.Size = 8;
 		ClientIconFontInfo.OutlineSettings.OutlineSize = 1;
 
@@ -404,7 +405,7 @@ void SActiveSession::Construct(const FArguments& InArgs, TSharedPtr<IConcertSync
 		LeaveSessionDef.OnExecute.BindLambda([this]() { OnClickLeaveSession(); });
 		LeaveSessionDef.IconStyle = TEXT("Concert.LeaveSession");
 
-		ConcertFrontendUtils::AppendButtons(StatusBar, ButtonDefs);
+		ConcertClientFrontendUtils::AppendButtons(StatusBar, ButtonDefs);
 	}
 
 	ChildSlot
@@ -709,7 +710,7 @@ const FButtonStyle& SActiveSession::GetConnectionIconStyle() const
 		}
 	}
 	
-	return FAppStyle::Get().GetWidgetStyle<FButtonStyle>(ConcertFrontendUtils::ButtonStyleNames[(int32)ButtonStyle]);
+	return FAppStyle::Get().GetWidgetStyle<FButtonStyle>(ConcertClientFrontendUtils::ButtonStyleNames[(int32)ButtonStyle]);
 }
 
 FSlateColor SActiveSession::GetConnectionIconColor() const
@@ -719,7 +720,7 @@ FSlateColor SActiveSession::GetConnectionIconColor() const
 
 FSlateFontInfo SActiveSession::GetConnectionIconFontInfo() const
 {
-	FSlateFontInfo ConnectionIconFontInfo = FAppStyle::Get().GetFontStyle(ConcertFrontendUtils::ButtonIconSyle);
+	FSlateFontInfo ConnectionIconFontInfo = FAppStyle::Get().GetFontStyle(ConcertClientFrontendUtils::ButtonIconSyle);
 	ConnectionIconFontInfo.OutlineSettings.OutlineSize = 1;
 	ConnectionIconFontInfo.OutlineSettings.OutlineColor = GetConnectionIconStyle().Pressed.TintColor.GetSpecifiedColor();
 
diff --git a/Engine/Plugins/Developer/Concert/MultiUserClient/Source/MultiUserClient/Private/Widgets/SConcertBrowser.cpp b/Engine/Plugins/Developer/Concert/MultiUserClient/Source/MultiUserClient/Private/Widgets/SConcertBrowser.cpp
--- a/Engine/Plugins/Developer/Concert/MultiUserClient/Source/MultiUserClient/Private/Widgets/SConcertBrowser.cpp
+++ b/Engine/Plugins/Developer/Concert/MultiUserClient/Source/MultiUserClient/Private/Widgets/SConcertBrowser.cpp
@@ -354,7 +354,7 @@ void FConcertClientSessionBrowserController::RestoreSession(const FGuid& ServerA
 	TSharedPtr<FConcertActivityStream> ActivityStream = MakeShared<FConcertActivityStream>(ConcertClient, ServerAdminEndpointId, SessionId, bRequestActivityDetails);
 
 	// The UI uses this function to read and consume the activity stream.
-	auto ReadActivitiesFn = [ActivityStream](TArray<TSharedPtr<FConcertClientSessionActivity>>& InOutActivities, int32& OutFetchCount, FText& OutErrorMsg)
+	auto ReadActivitiesFn = [ActivityStream](TArray<TSharedPtr<FConcertSessionActivity>>& InOutActivities, int32& OutFetchCount, FText& OutErrorMsg)
 	{
 		return ActivityStream->Read(InOutActivities, OutFetchCount, OutErrorMsg);
 	};
@@ -367,7 +367,7 @@ void FConcertClientSessionBrowserController::RestoreSession(const FGuid& ServerA
 
 	// Invoked if the client selects a point in time to recover.
 	TWeakPtr<IConcertClient, ESPMode::ThreadSafe> WeakClient = ConcertClient;
-	auto OnAcceptRestoreFn = [WeakClient, ServerAdminEndpointId, SessionId, RestoredName, SessionFilter](TSharedPtr<FConcertClientSessionActivity> SelectedRecoveryActivity)
+	auto OnAcceptRestoreFn = [WeakClient, ServerAdminEndpointId, SessionId, RestoredName, SessionFilter](TSharedPtr<FConcertSessionActivity> SelectedRecoveryActivity)
 	{
 		FConcertCopySessionArgs RestoreSessionArgs;
 		RestoreSessionArgs.bAutoConnect = true;
diff --git a/Engine/Plugins/Developer/Concert/MultiUserClient/Source/MultiUserClient/Private/Widgets/SSessionHistory.cpp b/Engine/Plugins/Developer/Concert/MultiUserClient/Source/MultiUserClient/Private/Widgets/SSessionHistory.cpp
--- a/Engine/Plugins/Developer/Concert/MultiUserClient/Source/MultiUserClient/Private/Widgets/SSessionHistory.cpp
+++ b/Engine/Plugins/Developer/Concert/MultiUserClient/Source/MultiUserClient/Private/Widgets/SSessionHistory.cpp
@@ -50,8 +50,8 @@ void SSessionHistory::Construct(const FArguments& InArgs, TSharedPtr<IConcertSyn
 	ActivityListViewOptions = MakeShared<FConcertSessionActivitiesOptions>();
 
 	SAssignNew(ActivityListView, SConcertSessionActivities)
-		.OnGetPackageEvent([this](const FConcertClientSessionActivity& Activity, FConcertSyncPackageEventMetaData& OutEvent) { return GetPackageEvent(Activity, OutEvent); })
-		.OnGetTransactionEvent([this](const FConcertClientSessionActivity& Activity) { return GetTransactionEvent(Activity); })
+		.OnGetPackageEvent([this](const FConcertSessionActivity& Activity, FConcertSyncPackageEventMetaData& OutEvent) { return GetPackageEvent(Activity, OutEvent); })
+		.OnGetTransactionEvent([this](const FConcertSessionActivity& Activity) { return GetTransactionEvent(Activity); })
 		.OnMapActivityToClient([this](FGuid ClientId){ return EndpointClientInfoMap.Find(ClientId); })
 		.HighlightText(this, &SSessionHistory::HighlightSearchedText)
 		.TimeFormat(ActivityListViewOptions.Get(), &FConcertSessionActivitiesOptions::GetTimeFormat)
@@ -151,14 +151,14 @@ void SSessionHistory::ReloadActivities()
 		const int64 LastActivityId = WorkspacePtr->GetLastActivityId();
 		const int64 FirstActivityIdToFetch = FMath::Max<int64>(1, LastActivityId - MaximumNumberOfActivities);
 
-		TArray<FConcertClientSessionActivity> FetchedActivities;
+		TArray<FConcertSessionActivity> FetchedActivities;
 		WorkspacePtr->GetActivities(FirstActivityIdToFetch, MaximumNumberOfActivities, EndpointClientInfoMap, FetchedActivities);
 
-		for (FConcertClientSessionActivity& FetchedActivity : FetchedActivities)
+		for (FConcertSessionActivity& FetchedActivity : FetchedActivities)
 		{
 			if (ConcertSessionHistoryUI::PackageNamePassesFilter(PackageNameFilter, FetchedActivity.ActivitySummary))
 			{
-				TSharedRef<FConcertClientSessionActivity> NewActivity = MakeShared<FConcertClientSessionActivity>(MoveTemp(FetchedActivity));
+				TSharedRef<FConcertSessionActivity> NewActivity = MakeShared<FConcertSessionActivity>(MoveTemp(FetchedActivity));
 				ActivityMap.Add(NewActivity->Activity.ActivityId, NewActivity);
 				ActivityListView->Append(NewActivity);
 			}
@@ -175,7 +175,7 @@ void SSessionHistory::HandleActivityAddedOrUpdated(const FConcertClientInfo& InC
 	{
 		EndpointClientInfoMap.Add(InActivity.EndpointId, InClientInfo);
 
-		if (TSharedPtr<FConcertClientSessionActivity> ExistingActivity = ActivityMap.FindRef(InActivity.ActivityId))
+		if (TSharedPtr<FConcertSessionActivity> ExistingActivity = ActivityMap.FindRef(InActivity.ActivityId))
 		{
 			ExistingActivity->Activity = InActivity;
 			ExistingActivity->ActivitySummary = MoveTemp(ActivitySummary);
@@ -183,7 +183,7 @@ void SSessionHistory::HandleActivityAddedOrUpdated(const FConcertClientInfo& InC
 		}
 		else
 		{
-			TSharedRef<FConcertClientSessionActivity> NewActivity = MakeShared<FConcertClientSessionActivity>();
+			TSharedRef<FConcertSessionActivity> NewActivity = MakeShared<FConcertSessionActivity>();
 			NewActivity->Activity = InActivity;
 			NewActivity->ActivitySummary = MoveTemp(ActivitySummary);
 
@@ -218,7 +218,7 @@ void SSessionHistory::RegisterWorkspaceHandler()
 	}
 }
 
-bool SSessionHistory::GetPackageEvent(const FConcertClientSessionActivity& Activity, FConcertSyncPackageEventMetaData& OutPackageEvent) const
+bool SSessionHistory::GetPackageEvent(const FConcertSessionActivity& Activity, FConcertSyncPackageEventMetaData& OutPackageEvent) const
 {
 	if (TSharedPtr<IConcertClientWorkspace> WorkspacePtr = Workspace.Pin())
 	{
@@ -229,7 +229,7 @@ bool SSessionHistory::GetPackageEvent(const FConcertClientSessionActivity& Activ
 	return false; // The data is not available.
 }
 
-TFuture<TOptional<FConcertSyncTransactionEvent>> SSessionHistory::GetTransactionEvent(const FConcertClientSessionActivity& Activity) const
+TFuture<TOptional<FConcertSyncTransactionEvent>> SSessionHistory::GetTransactionEvent(const FConcertSessionActivity& Activity) const
 {
 	if (TSharedPtr<IConcertClientWorkspace> WorkspacePtr = Workspace.Pin())
 	{
diff --git a/Engine/Plugins/Developer/Concert/MultiUserClient/Source/MultiUserClient/Private/Widgets/SSessionHistory.h b/Engine/Plugins/Developer/Concert/MultiUserClient/Source/MultiUserClient/Private/Widgets/SSessionHistory.h
--- a/Engine/Plugins/Developer/Concert/MultiUserClient/Source/MultiUserClient/Private/Widgets/SSessionHistory.h
+++ b/Engine/Plugins/Developer/Concert/MultiUserClient/Source/MultiUserClient/Private/Widgets/SSessionHistory.h
@@ -36,7 +36,7 @@ public:
 
 private:
 	/** Callback for selecting an activity in the list view. */
-	void HandleSelectionChanged(TSharedPtr<FConcertClientSessionActivity> InSessionActivity, ESelectInfo::Type SelectInfo);
+	void HandleSelectionChanged(TSharedPtr<FConcertSessionActivity> InSessionActivity, ESelectInfo::Type SelectInfo);
 
 	/** Fetches activities from the server and updates the list view. */
 	void ReloadActivities();
@@ -63,10 +63,10 @@ private:
 	FText HighlightSearchedText() const;
 
 	/** Returns the specified package event (without the package data itself) if available. */
-	bool GetPackageEvent(const FConcertClientSessionActivity& Activity, FConcertSyncPackageEventMetaData& OutPackageEvent) const;
+	bool GetPackageEvent(const FConcertSessionActivity& Activity, FConcertSyncPackageEventMetaData& OutPackageEvent) const;
 
 	/** Returns the specified package event if available. */
-	TFuture<TOptional<FConcertSyncTransactionEvent>> GetTransactionEvent(const FConcertClientSessionActivity& Activity) const;
+	TFuture<TOptional<FConcertSyncTransactionEvent>> GetTransactionEvent(const FConcertSessionActivity& Activity) const;
 
 private:
 
@@ -77,7 +77,7 @@ private:
 	TMap<FGuid, FConcertClientInfo> EndpointClientInfoMap;
 
 	/** Holds the map of activity IDs to Concert activities. */
-	TMap<int64, TSharedPtr<FConcertClientSessionActivity>> ActivityMap;
+	TMap<int64, TSharedPtr<FConcertSessionActivity>> ActivityMap;
 
 	/** Display the activity list. */
 	TSharedPtr<SConcertSessionActivities> ActivityListView;
