Adjust undo history and multi user UI to new UE5 conventions

diff --git a/Engine/Plugins/Developer/Concert/ConcertUI/ConcertSharedSlate/Source/ConcertSharedSlate/ConcertSharedSlate.Build.cs b/Engine/Plugins/Developer/Concert/ConcertUI/ConcertSharedSlate/Source/ConcertSharedSlate/ConcertSharedSlate.Build.cs
--- a/Engine/Plugins/Developer/Concert/ConcertUI/ConcertSharedSlate/Source/ConcertSharedSlate/ConcertSharedSlate.Build.cs
+++ b/Engine/Plugins/Developer/Concert/ConcertUI/ConcertSharedSlate/Source/ConcertSharedSlate/ConcertSharedSlate.Build.cs
@@ -9,24 +9,25 @@ namespace UnrealBuildTool.Rules
 			PublicDependencyModuleNames.AddRange(
 				new string[]
 				{
+					"Concert",
 					"Core",
 					"CoreUObject",
-					"Slate",
 					"EditorStyle",
-					"Concert"
+					"Slate",
 				}
 			);
 
 			PrivateDependencyModuleNames.AddRange(
 				new string[]
 				{
-					"SlateCore",
-					"InputCore",
-					"Projects",
 					"Concert",
 					"ConcertTransport",
 					"ConcertSyncCore",
-					"UndoHistory"
+					"InputCore",
+					"Projects",
+					"SlateCore",
+					"UndoHistory",
+					"ToolWidgets"
 				}
 			);
 		}
diff --git a/Engine/Plugins/Developer/Concert/ConcertUI/ConcertSharedSlate/Source/ConcertSharedSlate/Private/SConcertSessionRecovery.cpp b/Engine/Plugins/Developer/Concert/ConcertUI/ConcertSharedSlate/Source/ConcertSharedSlate/Private/SConcertSessionRecovery.cpp
--- a/Engine/Plugins/Developer/Concert/ConcertUI/ConcertSharedSlate/Source/ConcertSharedSlate/Private/SConcertSessionRecovery.cpp
+++ b/Engine/Plugins/Developer/Concert/ConcertUI/ConcertSharedSlate/Source/ConcertSharedSlate/Private/SConcertSessionRecovery.cpp
@@ -13,6 +13,8 @@
 #include "Widgets/Layout/SSeparator.h"
 #include "Widgets/Layout/SUniformGridPanel.h"
 #include "SConcertSessionActivities.h"
+#include "SNegativeActionButton.h"
+#include "SPositiveActionButton.h"
 
 #define LOCTEXT_NAMESPACE "SConcertSessionRecovery"
 
@@ -125,36 +127,19 @@ void SConcertSessionRecovery::Construct(const FArguments& InArgs)
 
 				+SUniformGridPanel::Slot(0, 0)
 				[
-					SNew(SButton)
-					.ForegroundColor(FLinearColor::White)
-					.ButtonStyle(FEditorStyle::Get(), TEXT("FlatButton.Success"))
+					SNew(SPositiveActionButton)
 					.ToolTipText(this, &SConcertSessionRecovery::GetRecoverAllButtonTooltip)
+					.Text(LOCTEXT("RecoverAll", "Recover All"))
+					.Icon(FEditorStyle::GetBrush("Icons.CircleArrowRight"))
 					.OnClicked(this, &SConcertSessionRecovery::OnRecoverAllClicked)
-					.HAlign(HAlign_Center)
-					.ContentPadding(FMargin(14, 3))
-					[
-						SNew(STextBlock)
-						.Text(LOCTEXT("RecoverAll", "Recover All"))
-						.Font( FEditorStyle::GetFontStyle("BoldFont"))
-						.ShadowOffset( FVector2D( 1.0f, 1.0f ) )
-					]
 				]
 
 				+SUniformGridPanel::Slot(1, 0)
 				[
-					SNew(SButton)
-					.ForegroundColor(FLinearColor::White)
-					.ButtonStyle(FEditorStyle::Get(), TEXT("FlatButton.Danger"))
+					SNew(SNegativeActionButton)
 					.ToolTipText(LOCTEXT("CancelRecoveryTooltip", "Discard any recoverable data for your assets and continue with their last saved state"))
+					.Text(LOCTEXT("Cancel", "Cancel"))
 					.OnClicked(this, &SConcertSessionRecovery::OnCancelRecoveryClicked)
-					.HAlign(HAlign_Center)
-					.ContentPadding(FMargin(14, 3))
-					[
-						SNew(STextBlock)
-						.Text(LOCTEXT("Cancel", "Cancel"))
-						.Font( FEditorStyle::GetFontStyle("BoldFont"))
-						.ShadowOffset( FVector2D( 1.0f, 1.0f ) )
-					]
 				]
 			]
 		]
@@ -173,18 +158,11 @@ TSharedPtr<SWidget> SConcertSessionRecovery::MakeRecoverThroughWidget(TWeakPtr<F
 			.HAlign(HAlign_Right)
 			.VAlign(VAlign_Center)
 			[
-				SNew(SButton)
-				.ForegroundColor(FLinearColor::White)
-				.ButtonStyle(FEditorStyle::Get(), TEXT("FlatButton.Success"))
+				SNew(SPositiveActionButton)
 				.Visibility_Lambda([this, Activity](){ return GetRecoverThroughButtonVisibility(Activity.Pin()); })
 				.OnClicked_Lambda([this, Activity](){ RecoverThrough(Activity.Pin()); return FReply::Handled(); })
 				.ToolTipText(this, &SConcertSessionRecovery::GetRecoverThroughButtonTooltip)
-				.ContentPadding(FMargin(20, 0))
-				[
-					SNew(STextBlock)
-					.Font(FEditorStyle::Get().GetFontStyle("FontAwesome.10"))
-					.Text(FEditorFontGlyphs::Arrow_Circle_O_Right)
-				]
+				.Icon(FEditorStyle::GetBrush("Icons.CircleArrowRight"))
 			];
 		}
 	}
diff --git a/Engine/Plugins/Developer/Concert/ConcertUI/ConcertSharedSlate/Source/ConcertSharedSlate/Private/SessionBrowser/SConcertSessionBrowser.cpp b/Engine/Plugins/Developer/Concert/ConcertUI/ConcertSharedSlate/Source/ConcertSharedSlate/Private/SessionBrowser/SConcertSessionBrowser.cpp
--- a/Engine/Plugins/Developer/Concert/ConcertUI/ConcertSharedSlate/Source/ConcertSharedSlate/Private/SessionBrowser/SConcertSessionBrowser.cpp
+++ b/Engine/Plugins/Developer/Concert/ConcertUI/ConcertSharedSlate/Source/ConcertSharedSlate/Private/SessionBrowser/SConcertSessionBrowser.cpp
@@ -272,7 +272,6 @@ TSharedRef<SWidget> SConcertSessionBrowser::MakeButtonBar(const FArguments& InAr
 	// New Session
 	RowBuilder.AddWidget(
 		ConcertBrowserUtils::MakeIconButton(
-			TEXT("SimpleButton"),
 			FConcertFrontendStyle::Get()->GetBrush("Concert.NewSession"),
 			LOCTEXT("NewButtonTooltip", "Create a new session"),
 			TAttribute<bool>(this, &SConcertSessionBrowser::IsNewButtonEnabledInternal),
@@ -287,7 +286,6 @@ TSharedRef<SWidget> SConcertSessionBrowser::MakeButtonBar(const FArguments& InAr
 	// Restore (Share the same slot as Join)
 	RowBuilder.AddWidget(
 		ConcertBrowserUtils::MakeIconButton(
-			TEXT("SimpleButton"),
 			FConcertFrontendStyle::Get()->GetBrush("Concert.RestoreSession"),
 			LOCTEXT("RestoreButtonTooltip", "Restore the selected session"),
 			TAttribute<bool>(this, &SConcertSessionBrowser::IsRestoreButtonEnabledInternal),
@@ -297,13 +295,13 @@ TSharedRef<SWidget> SConcertSessionBrowser::MakeButtonBar(const FArguments& InAr
 	);
 	// Archive.
 	RowBuilder.AddWidget(
-		ConcertBrowserUtils::MakeIconButton(TEXT("SimpleButton"), FConcertFrontendStyle::Get()->GetBrush("Concert.ArchiveSession"), LOCTEXT("ArchiveButtonTooltip", "Archive the selected session"),
+		ConcertBrowserUtils::MakeIconButton(FConcertFrontendStyle::Get()->GetBrush("Concert.ArchiveSession"), LOCTEXT("ArchiveButtonTooltip", "Archive the selected session"),
 			TAttribute<bool>(this, &SConcertSessionBrowser::IsArchiveButtonEnabledInternal),
 			FOnClicked::CreateSP(this, &SConcertSessionBrowser::OnArchiveButtonClicked))
 	);
 	// Delete.
 	RowBuilder.AddWidget(
-		ConcertBrowserUtils::MakeIconButton(TEXT("SimpleButton"), FConcertFrontendStyle::Get()->GetBrush("Concert.DeleteSession"), LOCTEXT("DeleteButtonTooltip", "Delete the selected session if permitted"),
+		ConcertBrowserUtils::MakeIconButton(FConcertFrontendStyle::Get()->GetBrush("Concert.DeleteSession"), LOCTEXT("DeleteButtonTooltip", "Delete the selected session if permitted"),
 			TAttribute<bool>(this, &SConcertSessionBrowser::IsDeleteButtonEnabledInternal),
 			FOnClicked::CreateSP(this, &SConcertSessionBrowser::OnDeleteButtonClicked))
 	);
diff --git a/Engine/Plugins/Developer/Concert/ConcertUI/ConcertSharedSlate/Source/ConcertSharedSlate/Private/SessionBrowser/SNewSessionRow.cpp b/Engine/Plugins/Developer/Concert/ConcertUI/ConcertSharedSlate/Source/ConcertSharedSlate/Private/SessionBrowser/SNewSessionRow.cpp
--- a/Engine/Plugins/Developer/Concert/ConcertUI/ConcertSharedSlate/Source/ConcertSharedSlate/Private/SessionBrowser/SNewSessionRow.cpp
+++ b/Engine/Plugins/Developer/Concert/ConcertUI/ConcertSharedSlate/Source/ConcertSharedSlate/Private/SessionBrowser/SNewSessionRow.cpp
@@ -6,6 +6,7 @@
 #include "SessionBrowser/ConcertSessionItem.h"
 
 #include "EditorFontGlyphs.h"
+#include "EditorStyleSet.h"
 #include "Framework/Application/SlateApplication.h"
 #include "Styling/AppStyle.h"
 #include "Widgets/Input/SEditableTextBox.h"
@@ -108,25 +109,21 @@ TSharedRef<SWidget> SNewSessionRow::GenerateWidgetForColumn(const FName& ColumnN
 				// 'Accept' button
 				+SUniformGridPanel::Slot(0, 0)
 				[
-					ConcertBrowserUtils::MakeIconButton(
-						TEXT("FlatButton.Success"),
-						FEditorFontGlyphs::Check,
+					ConcertBrowserUtils::MakePositiveActionButton(
+						FEditorStyle::GetBrush("Icons.Check"),
 						LOCTEXT("CreateCheckIconTooltip", "Create the session"),
 						TAttribute<bool>::Create([this]() { return !EditableSessionName->GetText().IsEmpty(); }),
-						FOnClicked::CreateRaw(this, &SNewSessionRow::OnAccept),
-						FSlateColor(FLinearColor::White))
+						FOnClicked::CreateRaw(this, &SNewSessionRow::OnAccept))
 				]
 
 				// 'Decline' button
 				+SUniformGridPanel::Slot(1, 0)
 				[
-					ConcertBrowserUtils::MakeIconButton(
-						TEXT("FlatButton.Danger"),
-						FEditorFontGlyphs::Times,
+					ConcertBrowserUtils::MakeNegativeActionButton(
+						FEditorStyle::GetBrush("Icons.X"),
 						LOCTEXT("CancelIconTooltip", "Cancel"),
 						true, // Always enabled.
-						FOnClicked::CreateRaw(this, &SNewSessionRow::OnDecline),
-						FSlateColor(FLinearColor::White))
+						FOnClicked::CreateRaw(this, &SNewSessionRow::OnDecline))
 				]
 			];
 	}
diff --git a/Engine/Plugins/Developer/Concert/ConcertUI/ConcertSharedSlate/Source/ConcertSharedSlate/Private/SessionBrowser/SSaveRestoreSessionRow.cpp b/Engine/Plugins/Developer/Concert/ConcertUI/ConcertSharedSlate/Source/ConcertSharedSlate/Private/SessionBrowser/SSaveRestoreSessionRow.cpp
--- a/Engine/Plugins/Developer/Concert/ConcertUI/ConcertSharedSlate/Source/ConcertSharedSlate/Private/SessionBrowser/SSaveRestoreSessionRow.cpp
+++ b/Engine/Plugins/Developer/Concert/ConcertUI/ConcertSharedSlate/Source/ConcertSharedSlate/Private/SessionBrowser/SSaveRestoreSessionRow.cpp
@@ -6,6 +6,7 @@
 #include "SessionBrowser/ConcertSessionItem.h"
 
 #include "EditorFontGlyphs.h"
+#include "EditorStyleSet.h"
 #include "Framework/Application/SlateApplication.h"
 #include "Internationalization/Regex.h"
 #include "Styling/AppStyle.h"
@@ -130,25 +131,21 @@ TSharedRef<SWidget> SSaveRestoreSessionRow::GenerateWidgetForColumn(const FName&
 			// 'Accept' button
 			+SUniformGridPanel::Slot(0, 0)
 			[
-				ConcertBrowserUtils::MakeIconButton(
-					TEXT("FlatButton.Success"),
-					FEditorFontGlyphs::Check,
+				ConcertBrowserUtils::MakePositiveActionButton(
+					FEditorStyle::GetBrush("Icons.Check"),
 					ItemPin->Type == FConcertSessionItem::EType::RestoreSession ? LOCTEXT("RestoreCheckIconTooltip", "Restore the session") : LOCTEXT("ArchiveCheckIconTooltip", "Archive the session"),
 					TAttribute<bool>::Create([this]() { return !EditableSessionName->GetText().IsEmpty(); }), // Enabled?
-					FOnClicked::CreateRaw(this, &SSaveRestoreSessionRow::OnAccept),
-					FSlateColor(FLinearColor::White))
+					FOnClicked::CreateRaw(this, &SSaveRestoreSessionRow::OnAccept))
 			]
 
 			// 'Cancel' button
 			+SUniformGridPanel::Slot(1, 0)
 			[
-				ConcertBrowserUtils::MakeIconButton(
-					TEXT("FlatButton.Danger"),
-					FEditorFontGlyphs::Times,
+				ConcertBrowserUtils::MakeNegativeActionButton(
+					FEditorStyle::GetBrush("Icons.X"),
 					LOCTEXT("CancelTooltip", "Cancel"),
 					true, // Enabled?
-					FOnClicked::CreateRaw(this, &SSaveRestoreSessionRow::OnDecline),
-					FSlateColor(FLinearColor::White))
+					FOnClicked::CreateRaw(this, &SSaveRestoreSessionRow::OnDecline))
 			]
 		];
 	}
diff --git a/Engine/Plugins/Developer/Concert/ConcertUI/ConcertSharedSlate/Source/ConcertSharedSlate/Public/SessionBrowser/ConcertBrowserUtils.h b/Engine/Plugins/Developer/Concert/ConcertUI/ConcertSharedSlate/Source/ConcertSharedSlate/Public/SessionBrowser/ConcertBrowserUtils.h
--- a/Engine/Plugins/Developer/Concert/ConcertUI/ConcertSharedSlate/Source/ConcertSharedSlate/Public/SessionBrowser/ConcertBrowserUtils.h
+++ b/Engine/Plugins/Developer/Concert/ConcertUI/ConcertSharedSlate/Source/ConcertSharedSlate/Public/SessionBrowser/ConcertBrowserUtils.h
@@ -7,6 +7,8 @@
 #include "ConcertMessageData.h"
 
 #include "EditorFontGlyphs.h"
+#include "SNegativeActionButton.h"
+#include "SPositiveActionButton.h"
 #include "Styling/AppStyle.h"
 #include "Internationalization/Regex.h"
 #include "Layout/Visibility.h"
@@ -33,43 +35,44 @@ namespace ConcertBrowserUtils
 	// The awesome font used to pick the icon displayed in the session list view 'Icon' column.
 	static const FName IconColumnFontName(TEXT("FontAwesome.9"));
 
-	/** Utility function used to create buttons displaying only an icon (using FontAwesome). */
-	inline TSharedRef<SButton> MakeIconButton(const FName& ButtonStyle, const TAttribute<FText>& GlyphIcon, const TAttribute<FText>& Tooltip, const TAttribute<bool>& EnabledAttribute, const FOnClicked& OnClicked,
-		const FSlateColor& ForegroundColor, const TAttribute<EVisibility>& Visibility = EVisibility::Visible, const TAttribute<FMargin>& ContentPadding = FMargin(3.0, 2.0), const FName FontStyle = IconColumnFontName)
+	/** Utility function used to create buttons displaying only an icon (using a brush) */
+	inline TSharedRef<SWidget> MakeIconButton(const TAttribute<const FSlateBrush*>& Icon, const TAttribute<FText>& Tooltip, const TAttribute<bool>& EnabledAttribute, const FOnClicked& OnClicked, const TAttribute<EVisibility>& Visibility = EVisibility::Visible)
 	{
 		return SNew(SButton)
-			.ButtonStyle(FAppStyle::Get(), ButtonStyle)
 			.OnClicked(OnClicked)
+			.ButtonStyle(FAppStyle::Get(), TEXT("SimpleButton"))
 			.ToolTipText(Tooltip)
-			.ContentPadding(ContentPadding)
 			.Visibility(Visibility)
 			.IsEnabled(EnabledAttribute)
 			.VAlign(VAlign_Center)
 			.HAlign(HAlign_Center)
 			[
-				SNew(STextBlock)
-				.Font(FAppStyle::Get().GetFontStyle(FontStyle))
-				.Text(GlyphIcon)
+				SNew(SImage)
+				.Image(Icon)
+				.ColorAndOpacity(FSlateColor::UseForeground())
 			];
 	}
 
-	/** Utility function used to create buttons displaying only an icon (using a brush) */
-	inline TSharedRef<SButton> MakeIconButton(const FName& ButtonStyle, const TAttribute<const FSlateBrush*>& Icon, const TAttribute<FText>& Tooltip, const TAttribute<bool>& EnabledAttribute, const FOnClicked& OnClicked, const TAttribute<EVisibility>& Visibility = EVisibility::Visible)
+	/** Utility function used to create buttons displaying a positive action*/
+	inline TSharedRef<SWidget> MakePositiveActionButton(const TAttribute<const FSlateBrush*>& Icon, const TAttribute<FText>& Tooltip, const TAttribute<bool>& EnabledAttribute, const FOnClicked& OnClicked, const TAttribute<EVisibility>& Visibility = EVisibility::Visible)
 	{
-		return SNew(SButton)
-			.ButtonStyle(FAppStyle::Get(), ButtonStyle)
+		return SNew(SPositiveActionButton)
 			.OnClicked(OnClicked)
 			.ToolTipText(Tooltip)
-			.ContentPadding(FMargin(0, 0))
 			.Visibility(Visibility)
 			.IsEnabled(EnabledAttribute)
-			.VAlign(VAlign_Center)
-			.HAlign(HAlign_Center)
-			[
-				SNew(SImage)
-				.Image(Icon)
-				.ColorAndOpacity(FSlateColor::UseForeground())
-			];
+			.Icon(Icon);
+	}
+
+	/** Utility function used to create buttons displaying a negative action */
+	inline TSharedRef<SWidget> MakeNegativeActionButton(const TAttribute<const FSlateBrush*>& Icon, const TAttribute<FText>& Tooltip, const TAttribute<bool>& EnabledAttribute, const FOnClicked& OnClicked, const TAttribute<EVisibility>& Visibility = EVisibility::Visible)
+	{
+		return SNew(SNegativeActionButton)
+			.OnClicked(OnClicked)
+			.ToolTipText(Tooltip)
+			.Visibility(Visibility)
+			.IsEnabled(EnabledAttribute)
+			.Icon(Icon);
 	}
 
 	/** Returns the tooltip shown when hovering the triangle with an exclamation icon when a server doesn't validate the version requirements. */
diff --git a/Engine/Plugins/Developer/Concert/MultiUserClient/Source/MultiUserClient/MultiUserClient.Build.cs b/Engine/Plugins/Developer/Concert/MultiUserClient/Source/MultiUserClient/MultiUserClient.Build.cs
--- a/Engine/Plugins/Developer/Concert/MultiUserClient/Source/MultiUserClient/MultiUserClient.Build.cs
+++ b/Engine/Plugins/Developer/Concert/MultiUserClient/Source/MultiUserClient/MultiUserClient.Build.cs
@@ -32,7 +32,8 @@ namespace UnrealBuildTool.Rules
 					"Slate",
 					"SlateCore",
 					"SourceControl",
-					"ToolMenus"
+					"ToolMenus",
+					"ToolWidgets"
 				}
 			);
 
diff --git a/Engine/Plugins/Developer/Concert/MultiUserClient/Source/MultiUserClient/Private/Widgets/SConcertBrowser.cpp b/Engine/Plugins/Developer/Concert/MultiUserClient/Source/MultiUserClient/Private/Widgets/SConcertBrowser.cpp
--- a/Engine/Plugins/Developer/Concert/MultiUserClient/Source/MultiUserClient/Private/Widgets/SConcertBrowser.cpp
+++ b/Engine/Plugins/Developer/Concert/MultiUserClient/Source/MultiUserClient/Private/Widgets/SConcertBrowser.cpp
@@ -1202,7 +1202,7 @@ void SConcertClientSessionBrowser::ExtendControlButtons(FExtender& Extender)
 				// Launch server.
 				+SOverlay::Slot()
 				[
-					ConcertBrowserUtils::MakeIconButton(TEXT("SimpleButton"),
+					ConcertBrowserUtils::MakeIconButton(
 						FConcertFrontendStyle::Get()->GetBrush("Concert.NewServer"),
 						LOCTEXT("LaunchServerTooltip", "Launch a Multi-User server on your computer unless one is already running.\nThe editor UDP messaging settings will be passed to the launching server.\nThe server will use the server port found in Multi-User client settings if non 0 or use the editor udp messaging port number + 1, if non 0."),
 						TAttribute<bool>(this, &SConcertClientSessionBrowser::IsLaunchServerButtonEnabled),
@@ -1213,7 +1213,6 @@ void SConcertClientSessionBrowser::ExtendControlButtons(FExtender& Extender)
 				+SOverlay::Slot()
 				[
 					ConcertBrowserUtils::MakeIconButton(
-						TEXT("SimpleButton"),
 						FConcertFrontendStyle::Get()->GetBrush("Concert.CloseServer"),
 						LOCTEXT("ShutdownServerTooltip", "Shutdown the Multi-User server running on this computer."),
 						true, // Always enabled, but might be collapsed.
@@ -1256,7 +1255,7 @@ void SConcertClientSessionBrowser::ExtendControlButtons(FExtender& Extender)
 				// Auto-join
 				+SOverlay::Slot()
 				[
-					ConcertBrowserUtils::MakeIconButton(TEXT("SimpleButton"),
+					ConcertBrowserUtils::MakeIconButton(
 						FConcertFrontendStyle::Get()->GetBrush("Concert.JoinDefaultSession"),
 						AutoJoinTooltip,
 						TAttribute<bool>(this, &SConcertClientSessionBrowser::IsAutoJoinButtonEnabled),
@@ -1269,7 +1268,7 @@ void SConcertClientSessionBrowser::ExtendControlButtons(FExtender& Extender)
 				// Cancel auto join.
 				+SOverlay::Slot()
 				[
-					ConcertBrowserUtils::MakeIconButton(TEXT("SimpleButton"),
+					ConcertBrowserUtils::MakeIconButton(
 						FConcertFrontendStyle::Get()->GetBrush("Concert.CancelAutoJoin"),
 						CancelAutoJoinTooltip,
 						TAttribute<bool>(this, &SConcertClientSessionBrowser::IsCancelAutoJoinButtonEnabled),
@@ -1281,7 +1280,7 @@ void SConcertClientSessionBrowser::ExtendControlButtons(FExtender& Extender)
 
 			// Join
 			Builder.AddWidget(
-				ConcertBrowserUtils::MakeIconButton(TEXT("SimpleButton"), FConcertFrontendStyle::Get()->GetBrush("Concert.JoinSession"), LOCTEXT("JoinButtonTooltip", "Join the selected session"),
+				ConcertBrowserUtils::MakeIconButton(FConcertFrontendStyle::Get()->GetBrush("Concert.JoinSession"), LOCTEXT("JoinButtonTooltip", "Join the selected session"),
 					TAttribute<bool>(this, &SConcertClientSessionBrowser::IsJoinButtonEnabled),
 					FOnClicked::CreateSP(this, &SConcertClientSessionBrowser::OnJoinButtonClicked),
 					TAttribute<EVisibility>::Create([this]() { return !SessionBrowser->IsRestoreButtonEnabled() ? EVisibility::Visible : EVisibility::Collapsed; })) // Default button shown if both join/restore are disabled.
diff --git a/Engine/Source/Editor/UndoHistoryEditor/Private/Widgets/SUndoHistory.cpp b/Engine/Source/Editor/UndoHistoryEditor/Private/Widgets/SUndoHistory.cpp
--- a/Engine/Source/Editor/UndoHistoryEditor/Private/Widgets/SUndoHistory.cpp
+++ b/Engine/Source/Editor/UndoHistoryEditor/Private/Widgets/SUndoHistory.cpp
@@ -95,7 +95,7 @@ void SUndoHistory::Construct( const FArguments& InArgs )
 														.FillWidth(80.0f)
 
 													+ SHeaderRow::Column("JumpToButton")
-														.FixedWidth(30.0f)
+														.FixedWidth(50.0f)
 												)
 										]
 								]
diff --git a/Engine/Source/Editor/UndoHistoryEditor/Private/Widgets/SUndoHistoryTableRow.h b/Engine/Source/Editor/UndoHistoryEditor/Private/Widgets/SUndoHistoryTableRow.h
--- a/Engine/Source/Editor/UndoHistoryEditor/Private/Widgets/SUndoHistoryTableRow.h
+++ b/Engine/Source/Editor/UndoHistoryEditor/Private/Widgets/SUndoHistoryTableRow.h
@@ -19,6 +19,7 @@
 #include "Widgets/Views/SExpanderArrow.h"
 #include "Widgets/Layout/SExpandableArea.h"
 #include "EditorStyleSet.h"
+#include "SPositiveActionButton.h"
 #include "Widgets/Views/STreeView.h"
 #include "Widgets/Input/SButton.h"
 
@@ -82,23 +83,11 @@ public:
 	{
 		if (ColumnName == "JumpToButton")
 		{
-			return SNew(SButton)
+			return SNew(SPositiveActionButton)
 				.ToolTipText(FText::FromString("Jump to this transaction"))
-				.ButtonStyle(FEditorStyle::Get(), "FlatButton.Success")
-				.ForegroundColor(FLinearColor::White)
-				.ContentPadding(FMargin(2.f, 2.f))
+				.Icon(FEditorStyle::GetBrush("Icons.CircleArrowRight"))
 				.OnClicked_Lambda([this]() { OnGotoTransactionClicked.ExecuteIfBound(TransactionId); return FReply::Handled(); })
-				.Visibility_Lambda([this]() { return this->IsHovered() ? EVisibility::Visible : EVisibility::Hidden; })
-				[
-					SNew(SBox)
-						.HAlign(HAlign_Center)
-						[
-							SNew(STextBlock)
-								.Font(FEditorStyle::Get().GetFontStyle("FontAwesome.10"))
-								.Text(FText::FromString(TEXT("\xf138"))) /*fa-chevron-circle-right*/
-								.Justification(ETextJustify::Center)
-						]
-				];
+				.Visibility_Lambda([this]() { return this->IsHovered() ? EVisibility::Visible : EVisibility::Hidden; });
 		}
 		else if (ColumnName == "Title")
 		{
diff --git a/Engine/Source/Editor/UndoHistoryEditor/UndoHistoryEditor.Build.cs b/Engine/Source/Editor/UndoHistoryEditor/UndoHistoryEditor.Build.cs
--- a/Engine/Source/Editor/UndoHistoryEditor/UndoHistoryEditor.Build.cs
+++ b/Engine/Source/Editor/UndoHistoryEditor/UndoHistoryEditor.Build.cs
@@ -22,6 +22,7 @@ namespace UnrealBuildTool.Rules
 					"InputCore",
 					"Slate",
 					"SlateCore",
+					"ToolWidgets",
 					"UndoHistory",
 					"UnrealEd",
 				}
