Move files to correct modules (step 2 of refactoring SSessionHistory)

diff --git a/Engine/Plugins/Developer/Concert/ConcertUI/ConcertServerUI/Source/ConcertServerUI/Private/Widgets/Session/ConcertServerSessionTab.cpp b/Engine/Plugins/Developer/Concert/ConcertUI/ConcertServerUI/Source/ConcertServerUI/Private/Widgets/Session/ConcertServerSessionTab.cpp
--- a/Engine/Plugins/Developer/Concert/ConcertUI/ConcertServerUI/Source/ConcertServerUI/Private/Widgets/Session/ConcertServerSessionTab.cpp
+++ b/Engine/Plugins/Developer/Concert/ConcertUI/ConcertServerUI/Source/ConcertServerUI/Private/Widgets/Session/ConcertServerSessionTab.cpp
@@ -5,7 +5,7 @@
 #include "IConcertSession.h"
 #include "Widgets/ConcertServerTabs.h"
 #include "Widgets/StatusBar/SConcertStatusBar.h"
-#include "SessionInspector/SConcertSessionInspector.h"
+#include "Widgets/Session/SConcertSessionInspector.h"
 
 FConcertServerSessionTab::FConcertServerSessionTab(TSharedRef<IConcertServerSession> InspectedSession, const TSharedRef<SWindow>& ConstructUnderWindow)
 	: InspectedSession(MoveTemp(InspectedSession))
diff --git a/Engine/Plugins/Developer/Concert/ConcertUI/ConcertServerUI/Source/ConcertServerUI/Private/Widgets/Session/SConcertSessionInspector.cpp b/Engine/Plugins/Developer/Concert/ConcertUI/ConcertServerUI/Source/ConcertServerUI/Private/Widgets/Session/SConcertSessionInspector.cpp
new file mode 100644
--- /dev/null
+++ b/Engine/Plugins/Developer/Concert/ConcertUI/ConcertServerUI/Source/ConcertServerUI/Private/Widgets/Session/SConcertSessionInspector.cpp
@@ -0,0 +1,132 @@
+﻿// Copyright Epic Games, Inc. All Rights Reserved.
+
+#include "Widgets/Session/SConcertSessionInspector.h"
+
+#include "Framework/Docking/TabManager.h"
+#include "Widgets/Docking/SDockTab.h"
+#include "Widgets/SWindow.h"
+
+#define LOCTEXT_NAMESPACE "SConcertSessionInspector"
+
+const FName SConcertSessionInspector::HistoryTabId("HistoryTabId");
+const FName SConcertSessionInspector::SessionContentTabId("SessionContentTabId");
+const FName SConcertSessionInspector::ConnectionMonitorTabId("ConnectionMonitorTabId");
+
+void SConcertSessionInspector::Construct(const FArguments& InArgs, const TSharedRef<SDockTab>& ConstructUnderMajorTab, const TSharedRef<SWindow>& ConstructUnderWindow)
+{
+	ChildSlot
+	[
+		SNew(SBorder)
+		.BorderImage(FAppStyle::Get().GetBrush("ToolPanel.GroupBorder"))
+		.Padding(FMargin(1.0f, 2.0f))
+		[
+			SNew(SVerticalBox)
+
+			// Toolbar
+			+SVerticalBox::Slot()
+			.AutoHeight()
+			[
+				SNullWidget::NullWidget
+			]
+
+			// Content
+			+SVerticalBox::Slot()
+			.FillHeight(1.f)
+			[
+				SNew(SBorder)
+				.BorderImage(FAppStyle::Get().GetBrush("Brushes.Title"))
+				.Padding(FMargin(0.f, 0.f, 0.f, 5.0f)) // Visually separate status bar from tabs
+				[
+					CreateTabs(ConstructUnderMajorTab, ConstructUnderWindow)
+				]
+			]
+
+			// Status bar
+			+SVerticalBox::Slot()
+			.AutoHeight()
+			.HAlign(HAlign_Fill)
+			[
+				InArgs._StatusBar.Widget
+			]
+		]
+	];
+}
+
+TSharedRef<SWidget> SConcertSessionInspector::CreateTabs(const TSharedRef<SDockTab>& ConstructUnderMajorTab, const TSharedRef<SWindow>& ConstructUnderWindow)
+{
+	TabManager = FGlobalTabmanager::Get()->NewTabManager(ConstructUnderMajorTab);
+	
+	TabManager->RegisterTabSpawner(HistoryTabId, FOnSpawnTab::CreateSP(this, &SConcertSessionInspector::SpawnActivityHistory))
+		.SetDisplayName(LOCTEXT("ActivityHistoryLabel", "History"));
+	
+	TabManager->RegisterTabSpawner(SessionContentTabId, FOnSpawnTab::CreateSP(this, &SConcertSessionInspector::SpawnSessionContent))
+		.SetDisplayName(LOCTEXT("SessionContentLabel", "Session Content"));
+	
+	TabManager->RegisterTabSpawner(ConnectionMonitorTabId, FOnSpawnTab::CreateSP(this, &SConcertSessionInspector::SpawnConnectionMonitor))
+		.SetDisplayName(LOCTEXT("ConnectionMonitorLabel", "Connection Monitor"));
+
+	const TSharedRef<FTabManager::FLayout> Layout = FTabManager::NewLayout("ConcertSessionLayout_v0.3")
+		->AddArea
+		(
+			FTabManager::NewPrimaryArea()
+				->SetOrientation(Orient_Horizontal)
+				->Split
+				(
+					FTabManager::NewStack()
+					->SetSizeCoefficient(0.5f)
+					->AddTab(HistoryTabId, ETabState::OpenedTab)
+				)
+				->Split
+				(
+					FTabManager::NewSplitter()
+					->SetSizeCoefficient(0.5f)
+					->SetOrientation(Orient_Vertical)
+					->Split
+					(
+						FTabManager::NewStack()
+						->SetSizeCoefficient(0.5f)
+						->AddTab(SessionContentTabId, ETabState::OpenedTab)
+					)
+					->Split
+					(
+						FTabManager::NewStack()
+						->SetSizeCoefficient(0.5f)
+						->AddTab(ConnectionMonitorTabId, ETabState::OpenedTab)
+					)
+				)
+		);
+
+	return TabManager->RestoreFrom(Layout, ConstructUnderWindow).ToSharedRef();
+}
+
+TSharedRef<SDockTab> SConcertSessionInspector::SpawnActivityHistory(const FSpawnTabArgs& Args)
+{
+	return SNew(SDockTab)
+		.Label(LOCTEXT("ActivityHistoryLabel", "History"))
+		.TabRole(PanelTab)
+		[
+			SNullWidget::NullWidget
+		];
+}
+
+TSharedRef<SDockTab> SConcertSessionInspector::SpawnSessionContent(const FSpawnTabArgs& Args)
+{
+	return SNew(SDockTab)
+		.Label(LOCTEXT("SessionContentLabel", "Session Content"))
+		.TabRole(PanelTab)
+		[
+			SNullWidget::NullWidget
+		];
+}
+
+TSharedRef<SDockTab> SConcertSessionInspector::SpawnConnectionMonitor(const FSpawnTabArgs& Args)
+{
+	return SNew(SDockTab)
+		.Label(LOCTEXT("ConnectionMonitorLabel", "Connection Monitor"))
+		.TabRole(PanelTab)
+		[
+			SNullWidget::NullWidget
+		];
+}
+
+#undef LOCTEXT_NAMESPACE
diff --git a/Engine/Plugins/Developer/Concert/ConcertUI/ConcertServerUI/Source/ConcertServerUI/Private/Widgets/Session/SConcertSessionInspector.h b/Engine/Plugins/Developer/Concert/ConcertUI/ConcertServerUI/Source/ConcertServerUI/Private/Widgets/Session/SConcertSessionInspector.h
new file mode 100644
--- /dev/null
+++ b/Engine/Plugins/Developer/Concert/ConcertUI/ConcertServerUI/Source/ConcertServerUI/Private/Widgets/Session/SConcertSessionInspector.h
@@ -0,0 +1,45 @@
+﻿// Copyright Epic Games, Inc. All Rights Reserved.
+
+#pragma once
+
+#include "CoreMinimal.h"
+#include "Widgets/SCompoundWidget.h"
+#include "Widgets/DeclarativeSyntaxSupport.h"
+
+class FTabManager;
+class FSpawnTabArgs;
+class SDockTab;
+class SWindow;
+
+/**
+ * Designed to be the content of a tab showing:
+ *  - activity history (transactions stored on the server as well as who made those transactions)
+ *  - session content (list of session data saved during a Multi-user session)
+ *  - connection monitor (details about the connected clients on the given session and network info)
+ * Implements view in model-view-controller pattern.
+ */
+class SConcertSessionInspector : public SCompoundWidget
+{
+public:
+
+	static const FName HistoryTabId;
+	static const FName SessionContentTabId;
+	static const FName ConnectionMonitorTabId;
+
+	SLATE_BEGIN_ARGS(SConcertSessionInspector) {}
+		SLATE_NAMED_SLOT(FArguments, StatusBar)
+	SLATE_END_ARGS()
+
+	void Construct(const FArguments& InArgs, const TSharedRef<SDockTab>& ConstructUnderMajorTab, const TSharedRef<SWindow>& ConstructUnderWindow);
+
+private:
+	
+	/** Holds the tab manager that manages the front-end's tabs. */
+	TSharedPtr<FTabManager> TabManager;
+
+	// Spawn tabs
+	TSharedRef<SWidget> CreateTabs(const TSharedRef<SDockTab>& ConstructUnderMajorTab, const TSharedRef<SWindow>& ConstructUnderWindow);
+	TSharedRef<SDockTab> SpawnActivityHistory(const FSpawnTabArgs& Args);
+	TSharedRef<SDockTab> SpawnSessionContent(const FSpawnTabArgs& Args);
+	TSharedRef<SDockTab> SpawnConnectionMonitor(const FSpawnTabArgs& Args);
+};
diff --git a/Engine/Plugins/Developer/Concert/ConcertUI/ConcertSharedSlate/Source/ConcertSharedSlate/Private/SessionBrowser/SSessionRow.cpp b/Engine/Plugins/Developer/Concert/ConcertUI/ConcertSharedSlate/Source/ConcertSharedSlate/Private/SessionBrowser/SSessionRow.cpp
--- a/Engine/Plugins/Developer/Concert/ConcertUI/ConcertSharedSlate/Source/ConcertSharedSlate/Private/SessionBrowser/SSessionRow.cpp
+++ b/Engine/Plugins/Developer/Concert/ConcertUI/ConcertSharedSlate/Source/ConcertSharedSlate/Private/SessionBrowser/SSessionRow.cpp
@@ -10,6 +10,7 @@
 #include "Styling/AppStyle.h"
 #include "Styling/CoreStyle.h"
 #include "Widgets/Layout/SBox.h"
+#include "Widgets/Layout/SSpacer.h"
 #include "Widgets/SBoxPanel.h"
 #include "Widgets/Text/SInlineEditableTextBlock.h"
 #include "Widgets/Text/STextBlock.h"
diff --git a/Engine/Plugins/Developer/Concert/ConcertUI/ConcertSharedSlate/Source/ConcertSharedSlate/Private/SessionHistory/AbstractSessionHistoryController.cpp b/Engine/Plugins/Developer/Concert/ConcertUI/ConcertSharedSlate/Source/ConcertSharedSlate/Private/SessionHistory/AbstractSessionHistoryController.cpp
new file mode 100644
--- /dev/null
+++ b/Engine/Plugins/Developer/Concert/ConcertUI/ConcertSharedSlate/Source/ConcertSharedSlate/Private/SessionHistory/AbstractSessionHistoryController.cpp
@@ -0,0 +1,32 @@
+﻿// Copyright Epic Games, Inc. All Rights Reserved.
+
+#include "SessionHistory/AbstractSessionHistoryController.h"
+
+#include "SessionHistory/SSessionHistory.h"
+#include "Widgets/DeclarativeSyntaxSupport.h"
+
+FAbstractSessionHistoryController::FAbstractSessionHistoryController(FName PackageFilter)
+	: SessionHistory(MakeSessionHistory(PackageFilter))
+{}
+
+void FAbstractSessionHistoryController::ReloadActivities()
+{
+	constexpr int64 MaximumNumberOfActivities = SSessionHistory::MaximumNumberOfActivities;
+	TMap<FGuid, FConcertClientInfo> EndpointClientInfoMap;
+	TArray<FConcertSessionActivity> FetchedActivities;
+	GetActivities(MaximumNumberOfActivities, EndpointClientInfoMap, FetchedActivities);
+
+	// MoveTemp strictly not needed - but it will be faster on Debug builds
+	GetSessionHistory()->ReloadActivities(
+		MoveTemp(EndpointClientInfoMap),
+		MoveTemp(FetchedActivities)
+		);
+}
+
+TSharedRef<SSessionHistory> FAbstractSessionHistoryController::MakeSessionHistory(FName PackageFilter) const
+{
+	return SNew(SSessionHistory)
+		.PackageFilter(PackageFilter)
+		.GetPackageEvent([this](const FConcertSessionActivity& Activity, FConcertSyncPackageEventMetaData& OutPackageEvent) { return GetPackageEvent(Activity, OutPackageEvent); })
+		.GetTransactionEvent([this](const FConcertSessionActivity& Activity) { return GetTransactionEvent(Activity); });
+}
diff --git a/Engine/Plugins/Developer/Concert/ConcertUI/ConcertSharedSlate/Source/ConcertSharedSlate/Private/SessionHistory/SSessionHistory.cpp b/Engine/Plugins/Developer/Concert/ConcertUI/ConcertSharedSlate/Source/ConcertSharedSlate/Private/SessionHistory/SSessionHistory.cpp
new file mode 100644
--- /dev/null
+++ b/Engine/Plugins/Developer/Concert/ConcertUI/ConcertSharedSlate/Source/ConcertSharedSlate/Private/SessionHistory/SSessionHistory.cpp
@@ -0,0 +1,166 @@
+// Copyright Epic Games, Inc. All Rights Reserved.
+
+#include "SessionHistory/SSessionHistory.h"
+
+#include "ConcertMessageData.h"
+#include "IConcertSession.h"
+#include "SConcertSessionActivities.h"
+
+#include "Widgets/Input/SSearchBox.h"
+#include "Widgets/Layout/SSeparator.h"
+
+#define LOCTEXT_NAMESPACE "SSessionHistory"
+
+namespace ConcertSessionHistoryUI
+{
+	bool PackageNamePassesFilter(const FName& PackageNameFilter, const TStructOnScope<FConcertSyncActivitySummary>& InActivitySummary)
+	{
+		if (PackageNameFilter.IsNone())
+		{
+			return true;
+		}
+
+		if (const FConcertSyncLockActivitySummary* Summary = InActivitySummary.Cast<FConcertSyncLockActivitySummary>())
+		{
+			return Summary->PrimaryPackageName == PackageNameFilter;
+		}
+
+		if (const FConcertSyncTransactionActivitySummary* Summary = InActivitySummary.Cast<FConcertSyncTransactionActivitySummary>())
+		{
+			return Summary->PrimaryPackageName == PackageNameFilter;
+		}
+
+		if (const FConcertSyncPackageActivitySummary* Summary = InActivitySummary.Cast<FConcertSyncPackageActivitySummary>())
+		{
+			return Summary->PackageName == PackageNameFilter;
+		}
+
+		return false;
+	}
+}
+
+void SSessionHistory::Construct(const FArguments& InArgs)
+{
+	PackageNameFilter = InArgs._PackageFilter;
+
+	ActivityMap.Reserve(MaximumNumberOfActivities);
+	ActivityListViewOptions = MakeShared<FConcertSessionActivitiesOptions>();
+
+	SAssignNew(ActivityListView, SConcertSessionActivities)
+		.OnGetPackageEvent(InArgs._GetPackageEvent)
+		.OnGetTransactionEvent(InArgs._GetTransactionEvent)
+		.OnMapActivityToClient([this](FGuid ClientId){ return EndpointClientInfoMap.Find(ClientId); })
+		.HighlightText(this, &SSessionHistory::HighlightSearchedText)
+		.TimeFormat(ActivityListViewOptions.Get(), &FConcertSessionActivitiesOptions::GetTimeFormat)
+		.ClientNameColumnVisibility(EVisibility::Visible)
+		.ClientAvatarColorColumnVisibility(EVisibility::Visible)
+		.OperationColumnVisibility(EVisibility::Visible)
+		.PackageColumnVisibility(EVisibility::Collapsed)
+		.ConnectionActivitiesVisibility(ActivityListViewOptions.Get(), &FConcertSessionActivitiesOptions::GetConnectionActivitiesVisibility)
+		.LockActivitiesVisibility(ActivityListViewOptions.Get(), &FConcertSessionActivitiesOptions::GetLockActivitiesVisibility)
+		.PackageActivitiesVisibility(ActivityListViewOptions.Get(), &FConcertSessionActivitiesOptions::GetPackageActivitiesVisibility)
+		.TransactionActivitiesVisibility(ActivityListViewOptions.Get(), &FConcertSessionActivitiesOptions::GetTransactionActivitiesVisibility)
+		.DetailsAreaVisibility(EVisibility::Visible)
+		.IsAutoScrollEnabled(true);
+
+	ChildSlot
+	[
+		SNew(SVerticalBox)
+		+SVerticalBox::Slot()
+		.AutoHeight()
+		.Padding(1, 1)
+		[
+			SAssignNew(SearchBox, SSearchBox)
+			.HintText(LOCTEXT("SearchHint", "Search..."))
+			.OnTextChanged(this, &SSessionHistory::OnSearchTextChanged)
+			.OnTextCommitted(this, &SSessionHistory::OnSearchTextCommitted)
+			.DelayChangeNotificationsWhileTyping(true)
+		]
+
+		+SVerticalBox::Slot()
+		[
+			ActivityListView.ToSharedRef()
+		]
+
+		+SVerticalBox::Slot()
+		.AutoHeight()
+		.Padding(4, 3)
+		[
+			SNew(SSeparator)
+		]
+
+		+SVerticalBox::Slot()
+		.AutoHeight()
+		.Padding(4, 0, 4, 3)
+		[
+			ActivityListViewOptions->MakeStatusBar(
+				TAttribute<int32>(ActivityListView.Get(), &SConcertSessionActivities::GetTotalActivityNum),
+				TAttribute<int32>(ActivityListView.Get(), &SConcertSessionActivities::GetDisplayedActivityNum))
+		]
+	];
+}
+
+void SSessionHistory::ReloadActivities(TMap<FGuid, FConcertClientInfo> InEndpointClientInfoMap, TArray<FConcertSessionActivity> InFetchedActivities)
+{
+	EndpointClientInfoMap = MoveTemp(InEndpointClientInfoMap);
+	ActivityMap.Reset();
+	ActivityListView->Reset(); // Careful, don't reset the shared ptr.
+
+	for (FConcertSessionActivity& FetchedActivity : InFetchedActivities)
+	{
+		if (ConcertSessionHistoryUI::PackageNamePassesFilter(PackageNameFilter, FetchedActivity.ActivitySummary))
+		{
+			TSharedRef<FConcertSessionActivity> NewActivity = MakeShared<FConcertSessionActivity>(MoveTemp(FetchedActivity));
+			ActivityMap.Add(NewActivity->Activity.ActivityId, NewActivity);
+			ActivityListView->Append(NewActivity);
+		}
+	}
+}
+
+void SSessionHistory::HandleActivityAddedOrUpdated(const FConcertClientInfo& InClientInfo, const FConcertSyncActivity& InActivity, const FStructOnScope& InActivitySummary)
+{
+	TStructOnScope<FConcertSyncActivitySummary> ActivitySummary;
+	ActivitySummary.InitializeFromChecked(InActivitySummary);
+
+	if (ConcertSessionHistoryUI::PackageNamePassesFilter(PackageNameFilter, ActivitySummary))
+	{
+		EndpointClientInfoMap.Add(InActivity.EndpointId, InClientInfo);
+
+		if (TSharedPtr<FConcertSessionActivity> ExistingActivity = ActivityMap.FindRef(InActivity.ActivityId))
+		{
+			ExistingActivity->Activity = InActivity;
+			ExistingActivity->ActivitySummary = MoveTemp(ActivitySummary);
+			ActivityListView->RequestRefresh();
+		}
+		else
+		{
+			TSharedRef<FConcertSessionActivity> NewActivity = MakeShared<FConcertSessionActivity>();
+			NewActivity->Activity = InActivity;
+			NewActivity->ActivitySummary = MoveTemp(ActivitySummary);
+
+			ActivityMap.Add(NewActivity->Activity.ActivityId, NewActivity);
+			ActivityListView->Append(NewActivity);
+		}
+	}
+}
+
+void SSessionHistory::OnSearchTextChanged(const FText& InSearchText)
+{
+	SearchedText = InSearchText;
+	SearchBox->SetError(ActivityListView->UpdateTextFilter(InSearchText));
+}
+
+void SSessionHistory::OnSearchTextCommitted(const FText& InSearchText, ETextCommit::Type CommitType)
+{
+	if (!InSearchText.EqualTo(SearchedText))
+	{
+		OnSearchTextChanged(InSearchText);
+	}
+}
+
+FText SSessionHistory::HighlightSearchedText() const
+{
+	return SearchedText;
+}
+
+#undef LOCTEXT_NAMESPACE /* SSessionHistory */
diff --git a/Engine/Plugins/Developer/Concert/ConcertUI/ConcertSharedSlate/Source/ConcertSharedSlate/Private/SessionInspector/SConcertSessionInspector.cpp b/Engine/Plugins/Developer/Concert/ConcertUI/ConcertSharedSlate/Source/ConcertSharedSlate/Private/SessionInspector/SConcertSessionInspector.cpp
deleted file mode 100644
diff --git a/Engine/Plugins/Developer/Concert/ConcertUI/ConcertSharedSlate/Source/ConcertSharedSlate/Public/SessionHistory/AbstractSessionHistoryController.h b/Engine/Plugins/Developer/Concert/ConcertUI/ConcertSharedSlate/Source/ConcertSharedSlate/Public/SessionHistory/AbstractSessionHistoryController.h
new file mode 100644
--- /dev/null
+++ b/Engine/Plugins/Developer/Concert/ConcertUI/ConcertSharedSlate/Source/ConcertSharedSlate/Public/SessionHistory/AbstractSessionHistoryController.h
@@ -0,0 +1,45 @@
+﻿// Copyright Epic Games, Inc. All Rights Reserved.
+
+#pragma once
+
+#include "CoreMinimal.h"
+#include "Async/Future.h"
+
+struct FConcertClientInfo;
+struct FConcertSyncTransactionEvent;
+struct FConcertSyncPackageEventMetaData;
+struct FConcertSessionActivity;
+class IConcertSyncClient;
+class SSessionHistory;
+
+/** Acts as controller of SSessionHistory in model-view-component pattern. */
+class FAbstractSessionHistoryController : public TSharedFromThis<FAbstractSessionHistoryController>
+{
+public:
+
+	FAbstractSessionHistoryController(FName PackageFilter = EName::None);
+	virtual ~FAbstractSessionHistoryController() = default;
+	
+	/** Fetches activities from the server and updates the list view. */
+	void ReloadActivities();
+	
+	const TSharedRef<SSessionHistory>& GetSessionHistory() const { return SessionHistory; }
+
+protected:
+
+	/** Gets the latest activities*/
+	virtual void GetActivities(int64 MaximumNumberOfActivities, TMap<FGuid, FConcertClientInfo>& EndpointClientInfoMap, TArray<FConcertSessionActivity>& FetchedActivities) const = 0;
+	
+	/** Returns the specified package event (without the package data itself) if available. */
+	virtual bool GetPackageEvent(const FConcertSessionActivity& Activity, FConcertSyncPackageEventMetaData& OutPackageEvent) const = 0;
+
+	/** Returns the specified package event if available. */
+	virtual TFuture<TOptional<FConcertSyncTransactionEvent>> GetTransactionEvent(const FConcertSessionActivity& Activity) const = 0;
+	
+private:
+
+	/** The widget being managed */
+	TSharedRef<SSessionHistory> SessionHistory;
+
+	TSharedRef<SSessionHistory> MakeSessionHistory(FName PackageFilter) const;
+};
\ No newline at end of file
diff --git a/Engine/Plugins/Developer/Concert/ConcertUI/ConcertSharedSlate/Source/ConcertSharedSlate/Public/SessionHistory/SSessionHistory.h b/Engine/Plugins/Developer/Concert/ConcertUI/ConcertSharedSlate/Source/ConcertSharedSlate/Public/SessionHistory/SSessionHistory.h
new file mode 100644
--- /dev/null
+++ b/Engine/Plugins/Developer/Concert/ConcertUI/ConcertSharedSlate/Source/ConcertSharedSlate/Public/SessionHistory/SSessionHistory.h
@@ -0,0 +1,81 @@
+// Copyright Epic Games, Inc. All Rights Reserved.
+
+#pragma once
+
+#include "CoreMinimal.h"
+#include "SConcertSessionActivities.h"
+#include "Widgets/DeclarativeSyntaxSupport.h"
+#include "Widgets/SCompoundWidget.h"
+
+class SSearchBox;
+class SConcertSessionActivities;
+class FConcertSessionActivitiesOptions;
+
+class SSessionHistory : public SCompoundWidget
+{
+public:
+
+	/** Maximum number of activities displayed on screen. */ 
+	static constexpr int64 MaximumNumberOfActivities = 1000;
+	
+	SLATE_BEGIN_ARGS(SSessionHistory) {}
+		SLATE_ARGUMENT(FName, PackageFilter)
+		SLATE_ARGUMENT(SConcertSessionActivities::FGetPackageEvent, GetPackageEvent)
+		SLATE_ARGUMENT(SConcertSessionActivities::FGetTransactionEvent, GetTransactionEvent)
+	SLATE_END_ARGS()
+
+	/**
+	* Constructs the Session History widget.
+	*
+	* @param InArgs The Slate argument list.
+	* @param ConstructUnderMajorTab The major tab which will contain the Session History widget.
+	* @param InConcertSyncClient Pointer on the concert sync client.
+	*/
+	void Construct(const FArguments& InArgs);
+
+	/** Fetches activities from the server and updates the list view. */
+	void ReloadActivities(TMap<FGuid, FConcertClientInfo> EndpointClientInfoMap, TArray<FConcertSessionActivity> FetchedActivities);
+	
+	/** Callback for handling the a new or updated activity item. */ 
+	void HandleActivityAddedOrUpdated(const FConcertClientInfo& InClientInfo, const FConcertSyncActivity& InActivity, const FStructOnScope& InActivitySummary);
+
+private:
+
+	/** Invoked when the text in the search box widget changes. */
+	void OnSearchTextChanged(const FText& InSearchText);
+
+	/** Invoked when the text in the search box widget is committed. */
+	void OnSearchTextCommitted(const FText& InFilterText, ETextCommit::Type CommitType);
+
+	/** Returns the text to highlight when the search bar has a text set. */
+	FText HighlightSearchedText() const;
+
+	/** Returns the specified package event (without the package data itself) if available. */
+	bool GetPackageEvent(const FConcertSessionActivity& Activity, FConcertSyncPackageEventMetaData& OutPackageEvent) const;
+
+	/** Returns the specified package event if available. */
+	TFuture<TOptional<FConcertSyncTransactionEvent>> GetTransactionEvent(const FConcertSessionActivity& Activity) const;
+
+private:
+	
+	/** Holds the map of endpoint IDs to client info. */
+	TMap<FGuid, FConcertClientInfo> EndpointClientInfoMap;
+
+	/** Holds the map of activity IDs to Concert activities. */
+	TMap<int64, TSharedPtr<FConcertSessionActivity>> ActivityMap;
+
+	/** Display the activity list. */
+	TSharedPtr<SConcertSessionActivities> ActivityListView;
+
+	/** Controls the activity list view options */
+	TSharedPtr<FConcertSessionActivitiesOptions> ActivityListViewOptions;
+
+	/** The widget used to enter the text to search. */
+	TSharedPtr<SSearchBox> SearchBox;
+
+	/** The searched text to highlight. */
+	FText SearchedText;
+
+	/** Used to limit activities to a given package only. */
+	FName PackageNameFilter;
+};
\ No newline at end of file
diff --git a/Engine/Plugins/Developer/Concert/ConcertUI/ConcertSharedSlate/Source/ConcertSharedSlate/Public/SessionHistory/SSessionHistoryWrapper.h b/Engine/Plugins/Developer/Concert/ConcertUI/ConcertSharedSlate/Source/ConcertSharedSlate/Public/SessionHistory/SSessionHistoryWrapper.h
new file mode 100644
--- /dev/null
+++ b/Engine/Plugins/Developer/Concert/ConcertUI/ConcertSharedSlate/Source/ConcertSharedSlate/Public/SessionHistory/SSessionHistoryWrapper.h
@@ -0,0 +1,32 @@
+﻿// Copyright Epic Games, Inc. All Rights Reserved.
+
+#pragma once
+
+#include "CoreMinimal.h"
+#include "AbstractSessionHistoryController.h"
+#include "SSessionHistory.h"
+#include "Widgets/SCompoundWidget.h"
+#include "Widgets/DeclarativeSyntaxSupport.h"
+
+/** Shows a SSessionHistory widget - keeping alive the controller as long as this widget lives. Useful for closeable tabs. */
+class SSessionHistoryWrapper : public SCompoundWidget
+{
+public:
+
+	SLATE_BEGIN_ARGS(SSessionHistoryWrapper) {}
+	SLATE_END_ARGS()
+
+	void Construct(const FArguments& InArgs, TSharedRef<FAbstractSessionHistoryController> InController)
+	{
+		Controller = InController;
+		ChildSlot
+		[
+			Controller->GetSessionHistory()
+		];
+	}
+
+private:
+
+	/** Manages SessionHistory */
+	TSharedPtr<FAbstractSessionHistoryController> Controller;
+};
diff --git a/Engine/Plugins/Developer/Concert/ConcertUI/ConcertSharedSlate/Source/ConcertSharedSlate/Public/SessionInspector/SConcertSessionInspector.h b/Engine/Plugins/Developer/Concert/ConcertUI/ConcertSharedSlate/Source/ConcertSharedSlate/Public/SessionInspector/SConcertSessionInspector.h
deleted file mode 100644
diff --git a/Engine/Plugins/Developer/Concert/MultiUserClient/Source/MultiUserClient/Private/ConcertWorkspaceUI.cpp b/Engine/Plugins/Developer/Concert/MultiUserClient/Source/MultiUserClient/Private/ConcertWorkspaceUI.cpp
--- a/Engine/Plugins/Developer/Concert/MultiUserClient/Source/MultiUserClient/Private/ConcertWorkspaceUI.cpp
+++ b/Engine/Plugins/Developer/Concert/MultiUserClient/Source/MultiUserClient/Private/ConcertWorkspaceUI.cpp
@@ -31,10 +31,10 @@
 #include "Framework/MultiBox/MultiBoxExtender.h"
 #include "Framework/Notifications/NotificationManager.h"
 
+#include "SessionHistory/SSessionHistoryWrapper.h"
 #include "ToolMenuContext.h"
 #include "ToolMenus.h"
 #include "Widgets/ClientSessionHistoryController.h"
-
 #include "Widgets/DeclarativeSyntaxSupport.h"
 #include "Widgets/SConcertSandboxPersistWidget.h"
 #include "Widgets/SCompoundWidget.h"
@@ -44,7 +44,6 @@
 #include "Widgets/Images/SImage.h"
 #include "Widgets/Text/STextBlock.h"
 #include "Widgets/Docking/SDockTab.h"
-#include "Widgets/SSessionHistoryWrapper.h"
 #include "Widgets/Notifications/SNotificationList.h"
 
 #define LOCTEXT_NAMESPACE "ConcertFrontend"
diff --git a/Engine/Plugins/Developer/Concert/MultiUserClient/Source/MultiUserClient/Private/Widgets/AbstractSessionHistoryController.cpp b/Engine/Plugins/Developer/Concert/MultiUserClient/Source/MultiUserClient/Private/Widgets/AbstractSessionHistoryController.cpp
deleted file mode 100644
diff --git a/Engine/Plugins/Developer/Concert/MultiUserClient/Source/MultiUserClient/Private/Widgets/AbstractSessionHistoryController.h b/Engine/Plugins/Developer/Concert/MultiUserClient/Source/MultiUserClient/Private/Widgets/AbstractSessionHistoryController.h
deleted file mode 100644
diff --git a/Engine/Plugins/Developer/Concert/MultiUserClient/Source/MultiUserClient/Private/Widgets/ClientSessionHistoryController.cpp b/Engine/Plugins/Developer/Concert/MultiUserClient/Source/MultiUserClient/Private/Widgets/ClientSessionHistoryController.cpp
--- a/Engine/Plugins/Developer/Concert/MultiUserClient/Source/MultiUserClient/Private/Widgets/ClientSessionHistoryController.cpp
+++ b/Engine/Plugins/Developer/Concert/MultiUserClient/Source/MultiUserClient/Private/Widgets/ClientSessionHistoryController.cpp
@@ -6,7 +6,7 @@
 #include "IConcertClientWorkspace.h"
 #include "IConcertSyncClient.h"
 
-#include "SSessionHistory.h"
+#include "SessionHistory/SSessionHistory.h"
 
 FClientSessionHistoryController::FClientSessionHistoryController(TSharedRef<IConcertSyncClient> Client, FName PackageFilter)
 	: FAbstractSessionHistoryController(PackageFilter)
diff --git a/Engine/Plugins/Developer/Concert/MultiUserClient/Source/MultiUserClient/Private/Widgets/ClientSessionHistoryController.h b/Engine/Plugins/Developer/Concert/MultiUserClient/Source/MultiUserClient/Private/Widgets/ClientSessionHistoryController.h
--- a/Engine/Plugins/Developer/Concert/MultiUserClient/Source/MultiUserClient/Private/Widgets/ClientSessionHistoryController.h
+++ b/Engine/Plugins/Developer/Concert/MultiUserClient/Source/MultiUserClient/Private/Widgets/ClientSessionHistoryController.h
@@ -3,7 +3,7 @@
 #pragma once
 
 #include "CoreMinimal.h"
-#include "AbstractSessionHistoryController.h"
+#include "SessionHistory/AbstractSessionHistoryController.h"
 
 struct FConcertSyncActivity;
 struct FConcertClientInfo;
diff --git a/Engine/Plugins/Developer/Concert/MultiUserClient/Source/MultiUserClient/Private/Widgets/SActiveSession.cpp b/Engine/Plugins/Developer/Concert/MultiUserClient/Source/MultiUserClient/Private/Widgets/SActiveSession.cpp
--- a/Engine/Plugins/Developer/Concert/MultiUserClient/Source/MultiUserClient/Private/Widgets/SActiveSession.cpp
+++ b/Engine/Plugins/Developer/Concert/MultiUserClient/Source/MultiUserClient/Private/Widgets/SActiveSession.cpp
@@ -17,14 +17,14 @@
 #include "FileHelpers.h"
 #include "Misc/PackageName.h"
 #include "Styling/AppStyle.h"
-#include "SSessionHistory.h"
+#include "SessionHistory/SSessionHistory.h"
+#include "SessionHistory/SSessionHistoryWrapper.h"
 #include "Widgets/Text/STextBlock.h"
 #include "Widgets/Images/SImage.h"
 #include "Widgets/Input/SHyperlink.h"
 #include "Widgets/Layout/SExpandableArea.h"
 #include "Widgets/Views/STableViewBase.h"
 #include "Widgets/Views/STableRow.h"
-#include "Widgets/SSessionHistoryWrapper.h"
 
 #define LOCTEXT_NAMESPACE "SActiveSession"
 
diff --git a/Engine/Plugins/Developer/Concert/MultiUserClient/Source/MultiUserClient/Private/Widgets/SSessionHistory.cpp b/Engine/Plugins/Developer/Concert/MultiUserClient/Source/MultiUserClient/Private/Widgets/SSessionHistory.cpp
deleted file mode 100644
diff --git a/Engine/Plugins/Developer/Concert/MultiUserClient/Source/MultiUserClient/Private/Widgets/SSessionHistory.h b/Engine/Plugins/Developer/Concert/MultiUserClient/Source/MultiUserClient/Private/Widgets/SSessionHistory.h
deleted file mode 100644
diff --git a/Engine/Plugins/Developer/Concert/MultiUserClient/Source/MultiUserClient/Private/Widgets/SSessionHistoryWrapper.h b/Engine/Plugins/Developer/Concert/MultiUserClient/Source/MultiUserClient/Private/Widgets/SSessionHistoryWrapper.h
deleted file mode 100644
