Fix linker issues

diff --git a/Engine/Plugins/Developer/Concert/ConcertUI/ConcertSharedSlate/Source/ConcertSharedSlate/Public/SessionHistory/AbstractSessionHistoryController.h b/Engine/Plugins/Developer/Concert/ConcertUI/ConcertSharedSlate/Source/ConcertSharedSlate/Public/SessionHistory/AbstractSessionHistoryController.h
--- a/Engine/Plugins/Developer/Concert/ConcertUI/ConcertSharedSlate/Source/ConcertSharedSlate/Public/SessionHistory/AbstractSessionHistoryController.h
+++ b/Engine/Plugins/Developer/Concert/ConcertUI/ConcertSharedSlate/Source/ConcertSharedSlate/Public/SessionHistory/AbstractSessionHistoryController.h
@@ -13,7 +13,7 @@ class IConcertSyncClient;
 class SSessionHistory;
 
 /** Acts as controller of SSessionHistory in model-view-component pattern. */
-class FAbstractSessionHistoryController : public TSharedFromThis<FAbstractSessionHistoryController>
+class CONCERTSHAREDSLATE_API FAbstractSessionHistoryController : public TSharedFromThis<FAbstractSessionHistoryController>
 {
 public:
 
diff --git a/Engine/Plugins/Developer/Concert/ConcertUI/ConcertSharedSlate/Source/ConcertSharedSlate/Public/SessionHistory/SSessionHistory.h b/Engine/Plugins/Developer/Concert/ConcertUI/ConcertSharedSlate/Source/ConcertSharedSlate/Public/SessionHistory/SSessionHistory.h
--- a/Engine/Plugins/Developer/Concert/ConcertUI/ConcertSharedSlate/Source/ConcertSharedSlate/Public/SessionHistory/SSessionHistory.h
+++ b/Engine/Plugins/Developer/Concert/ConcertUI/ConcertSharedSlate/Source/ConcertSharedSlate/Public/SessionHistory/SSessionHistory.h
@@ -11,7 +11,7 @@ class SSearchBox;
 class SConcertSessionActivities;
 class FConcertSessionActivitiesOptions;
 
-class SSessionHistory : public SCompoundWidget
+class CONCERTSHAREDSLATE_API SSessionHistory : public SCompoundWidget
 {
 public:
 
diff --git a/Engine/Plugins/Developer/Concert/ConcertUI/ConcertSharedSlate/Source/ConcertSharedSlate/Public/SessionHistory/SSessionHistoryWrapper.h b/Engine/Plugins/Developer/Concert/ConcertUI/ConcertSharedSlate/Source/ConcertSharedSlate/Public/SessionHistory/SSessionHistoryWrapper.h
--- a/Engine/Plugins/Developer/Concert/ConcertUI/ConcertSharedSlate/Source/ConcertSharedSlate/Public/SessionHistory/SSessionHistoryWrapper.h
+++ b/Engine/Plugins/Developer/Concert/ConcertUI/ConcertSharedSlate/Source/ConcertSharedSlate/Public/SessionHistory/SSessionHistoryWrapper.h
@@ -9,7 +9,7 @@
 #include "Widgets/DeclarativeSyntaxSupport.h"
 
 /** Shows a SSessionHistory widget - keeping alive the controller as long as this widget lives. Useful for closeable tabs. */
-class SSessionHistoryWrapper : public SCompoundWidget
+class CONCERTSHAREDSLATE_API SSessionHistoryWrapper : public SCompoundWidget
 {
 public:
 
