Fix exporting inlined function

diff --git a/Engine/Plugins/Developer/Concert/ConcertUI/ConcertSharedSlate/Source/ConcertSharedSlate/Private/SessionHistory/SSessionHistoryWrapper.cpp b/Engine/Plugins/Developer/Concert/ConcertUI/ConcertSharedSlate/Source/ConcertSharedSlate/Private/SessionHistory/SSessionHistoryWrapper.cpp
new file mode 100644
--- /dev/null
+++ b/Engine/Plugins/Developer/Concert/ConcertUI/ConcertSharedSlate/Source/ConcertSharedSlate/Private/SessionHistory/SSessionHistoryWrapper.cpp
@@ -0,0 +1,12 @@
+﻿// Copyright Epic Games, Inc. All Rights Reserved.
+
+#include "SessionHistory/SSessionHistoryWrapper.h"
+
+void SSessionHistoryWrapper::Construct(const FArguments& InArgs, TSharedRef<FAbstractSessionHistoryController> InController)
+{
+	Controller = InController;
+	ChildSlot
+	[
+		Controller->GetSessionHistory()
+	];
+}
\ No newline at end of file
diff --git a/Engine/Plugins/Developer/Concert/ConcertUI/ConcertSharedSlate/Source/ConcertSharedSlate/Public/SessionHistory/SSessionHistoryWrapper.h b/Engine/Plugins/Developer/Concert/ConcertUI/ConcertSharedSlate/Source/ConcertSharedSlate/Public/SessionHistory/SSessionHistoryWrapper.h
--- a/Engine/Plugins/Developer/Concert/ConcertUI/ConcertSharedSlate/Source/ConcertSharedSlate/Public/SessionHistory/SSessionHistoryWrapper.h
+++ b/Engine/Plugins/Developer/Concert/ConcertUI/ConcertSharedSlate/Source/ConcertSharedSlate/Public/SessionHistory/SSessionHistoryWrapper.h
@@ -16,14 +16,7 @@ public:
 	SLATE_BEGIN_ARGS(SSessionHistoryWrapper) {}
 	SLATE_END_ARGS()
 
-	void Construct(const FArguments& InArgs, TSharedRef<FAbstractSessionHistoryController> InController)
-	{
-		Controller = InController;
-		ChildSlot
-		[
-			Controller->GetSessionHistory()
-		];
-	}
+	void Construct(const FArguments& InArgs, TSharedRef<FAbstractSessionHistoryController> InController);
 
 private:
 
