More Non-unity fixes

diff --git a/Engine/Plugins/Developer/Concert/ConcertUI/ConcertSharedSlate/Source/ConcertSharedSlate/Private/SessionBrowser/SSaveRestoreSessionRow.h b/Engine/Plugins/Developer/Concert/ConcertUI/ConcertSharedSlate/Source/ConcertSharedSlate/Private/SessionBrowser/SSaveRestoreSessionRow.h
--- a/Engine/Plugins/Developer/Concert/ConcertUI/ConcertSharedSlate/Source/ConcertSharedSlate/Private/SessionBrowser/SSaveRestoreSessionRow.h
+++ b/Engine/Plugins/Developer/Concert/ConcertUI/ConcertSharedSlate/Source/ConcertSharedSlate/Private/SessionBrowser/SSaveRestoreSessionRow.h
@@ -5,6 +5,7 @@
 #include "CoreMinimal.h"
 #include "Misc/Attribute.h"
 #include "SessionBrowser/ConcertSessionItem.h"
+#include "Widgets/Views/SListView.h"
 #include "Widgets/DeclarativeSyntaxSupport.h"
 #include "Widgets/Views/STableRow.h"
 
diff --git a/Engine/Plugins/Developer/Concert/MultiUserClient/Source/MultiUserClient/Private/Widgets/SConcertBrowser.cpp b/Engine/Plugins/Developer/Concert/MultiUserClient/Source/MultiUserClient/Private/Widgets/SConcertBrowser.cpp
--- a/Engine/Plugins/Developer/Concert/MultiUserClient/Source/MultiUserClient/Private/Widgets/SConcertBrowser.cpp
+++ b/Engine/Plugins/Developer/Concert/MultiUserClient/Source/MultiUserClient/Private/Widgets/SConcertBrowser.cpp
@@ -27,6 +27,7 @@
 #include "Internationalization/Regex.h"
 #include "Misc/AsyncTaskNotification.h"
 #include "Misc/TextFilter.h"
+#include "Misc/MessageDialog.h"
 #include "Styling/AppStyle.h"
 #include "Textures/SlateIcon.h"
 #include "Widgets/Images/SImage.h"
