Add content viewer for session inspector
Unshelved from pending changelist '19318632':

diff --git a/Engine/Plugins/Developer/Concert/ConcertSync/ConcertSyncClient/Source/ConcertSyncClient/Private/ConcertClientWorkspace.cpp b/Engine/Plugins/Developer/Concert/ConcertSync/ConcertSyncClient/Source/ConcertSyncClient/Private/ConcertClientWorkspace.cpp
--- a/Engine/Plugins/Developer/Concert/ConcertSync/ConcertSyncClient/Source/ConcertSyncClient/Private/ConcertClientWorkspace.cpp
+++ b/Engine/Plugins/Developer/Concert/ConcertSync/ConcertSyncClient/Source/ConcertSyncClient/Private/ConcertClientWorkspace.cpp
@@ -425,7 +425,7 @@ void FConcertClientWorkspace::GetActivities(const int64 FirstActivityIdToFetch,
 			OutActivities.Emplace(MoveTemp(InActivity), MoveTemp(ActivitySummary));
 		}
 
-		return true;
+		return EBreakBehavior::Continue;
 	});
 }
 
diff --git a/Engine/Plugins/Developer/Concert/ConcertSync/ConcertSyncCore/Source/ConcertSyncCore/Private/ConcertSyncSessionDatabase.cpp b/Engine/Plugins/Developer/Concert/ConcertSync/ConcertSyncCore/Source/ConcertSyncCore/Private/ConcertSyncSessionDatabase.cpp
--- a/Engine/Plugins/Developer/Concert/ConcertSync/ConcertSyncCore/Source/ConcertSyncCore/Private/ConcertSyncSessionDatabase.cpp
+++ b/Engine/Plugins/Developer/Concert/ConcertSync/ConcertSyncCore/Source/ConcertSyncCore/Private/ConcertSyncSessionDatabase.cpp
@@ -1166,7 +1166,7 @@ public:
 	}
 
 	/** Get the activity data from activities for all activities of event_type */
-	SQLITE_PREPARED_STATEMENT(FGetAllActivityDataForEventType, "SELECT activity_id, endpoint_id, event_time, event_type, event_id, event_summary_type, event_summary_flags, event_summary_size_bytes, event_summary_data FROM activities WHERE event_type = ?1 ORDER BY activity_id;", SQLITE_PREPARED_STATEMENT_COLUMNS(int64, FGuid, FDateTime, int64, FName, int32, int32, TArray<uint8>), SQLITE_PREPARED_STATEMENT_BINDINGS(EConcertSyncActivityEventType));
+	SQLITE_PREPARED_STATEMENT(FGetAllActivityDataForEventType, "SELECT activity_id, endpoint_id, event_time, event_id, event_summary_type, event_summary_flags, event_summary_size_bytes, event_summary_data FROM activities WHERE event_type = ?1 ORDER BY activity_id;", SQLITE_PREPARED_STATEMENT_COLUMNS(int64, FGuid, FDateTime, int64, FName, int32, int32, TArray<uint8>), SQLITE_PREPARED_STATEMENT_BINDINGS(EConcertSyncActivityEventType));
 	FGetAllActivityDataForEventType Statement_GetAllActivityDataForEventType;
 	bool GetAllActivityDataForEventType(const EConcertSyncActivityEventType InEventType, TFunctionRef<ESQLitePreparedStatementExecuteRowResult(int64, const FGuid&, FDateTime, int64, FConcertSessionSerializedPayload&&)> InCallback)
 	{
@@ -1796,7 +1796,7 @@ bool FConcertSyncSessionDatabase::GetTransactionActivity(const int64 InActivityI
 		&& GetTransactionEvent(OutTransactionActivity.EventId, OutTransactionActivity.EventData);
 }
 
-bool FConcertSyncSessionDatabase::GetPackageActivity(const int64 InActivityId, const TFunctionRef<void(FConcertSyncActivity&&, FConcertSyncPackageEventData&)>& PackageActivityFn) const
+bool FConcertSyncSessionDatabase::GetPackageActivity(const int64 InActivityId, FConsumePackageActivityFunc PackageActivityFn) const
 {
 	FConcertSyncActivity PackageActivityBasePart;
 
@@ -1844,7 +1844,7 @@ bool FConcertSyncSessionDatabase::GetTransactionActivityForEvent(const int64 InT
 		&& GetTransactionEvent(InTransactionEventId, OutTransactionActivity.EventData);
 }
 
-bool FConcertSyncSessionDatabase::GetPackageActivityForEvent(const int64 InPackageEventId, const TFunctionRef<void(FConcertSyncActivity&&, FConcertSyncPackageEventData&)>& PackageActivityFn) const
+bool FConcertSyncSessionDatabase::GetPackageActivityForEvent(const int64 InPackageEventId, FIteratePackageActivityFunc PackageActivityFn) const
 {
 	FConcertSyncActivity PackageActivityBasePart;
 
@@ -1857,7 +1857,7 @@ bool FConcertSyncSessionDatabase::GetPackageActivityForEvent(const int64 InPacka
 		&& GetPackageEvent(InPackageEventId, GetPackageActivityEventFn); // Pull the package specific part and callback
 }
 
-bool FConcertSyncSessionDatabase::EnumerateActivities(TFunctionRef<bool(FConcertSyncActivity&&)> InCallback) const
+bool FConcertSyncSessionDatabase::EnumerateActivities(FIterateActivityFunc InCallback) const
 {
 	return Statements->GetAllActivityData([this, &InCallback](const int64 InActivityId, const FGuid& InEndpointId, const FDateTime InEventTime, const EConcertSyncActivityEventType InEventType, const int64 InEventId, FConcertSessionSerializedPayload&& InEventSummary)
 	{
@@ -1869,7 +1869,7 @@ bool FConcertSyncSessionDatabase::EnumerateActivities(TFunctionRef<bool(FConcert
 		Activity.EventType = InEventType;
 		Activity.EventId = InEventId;
 		Activity.EventSummary = MoveTemp(InEventSummary);
-		return InCallback(MoveTemp(Activity))
+		return InCallback(MoveTemp(Activity)) == EBreakBehavior::Continue
 			? ESQLitePreparedStatementExecuteRowResult::Continue
 			: ESQLitePreparedStatementExecuteRowResult::Stop;
 	});
@@ -1941,7 +1941,7 @@ bool FConcertSyncSessionDatabase::EnumerateTransactionActivities(TFunctionRef<bo
 	});
 }
 
-bool FConcertSyncSessionDatabase::EnumeratePackageActivities(const TFunctionRef<bool(FConcertSyncActivity&&, FConcertSyncPackageEventData&)>& InCallback) const
+bool FConcertSyncSessionDatabase::EnumeratePackageActivities(FIteratePackageActivityFunc InCallback) const
 {
 	return Statements->GetAllActivityDataForEventType(EConcertSyncActivityEventType::Package, [this, &InCallback](const int64 InActivityId, const FGuid& InEndpointId, const FDateTime InEventTime, const int64 InEventId, FConcertSessionSerializedPayload&& InEventSummary)
 	{
@@ -1957,7 +1957,7 @@ bool FConcertSyncSessionDatabase::EnumeratePackageActivities(const TFunctionRef<
 		ESQLitePreparedStatementExecuteRowResult Result = ESQLitePreparedStatementExecuteRowResult::Error;
 		GetPackageEvent(PackageActivity.EventId, [&InCallback, &Result, BasePart = MoveTemp(PackageActivity)](FConcertSyncPackageEventData& PackageEventPart) mutable
 		{
-			Result = InCallback(MoveTemp(BasePart), PackageEventPart)
+			Result = InCallback(MoveTemp(BasePart), PackageEventPart) == EBreakBehavior::Continue
 				? ESQLitePreparedStatementExecuteRowResult::Continue
 				: ESQLitePreparedStatementExecuteRowResult::Stop;
 		});
@@ -1965,7 +1965,7 @@ bool FConcertSyncSessionDatabase::EnumeratePackageActivities(const TFunctionRef<
 	});
 }
 
-bool FConcertSyncSessionDatabase::EnumerateActivitiesForEventType(const EConcertSyncActivityEventType InEventType, TFunctionRef<bool(FConcertSyncActivity&&)> InCallback) const
+bool FConcertSyncSessionDatabase::EnumerateActivitiesForEventType(const EConcertSyncActivityEventType InEventType, FIterateActivityFunc InCallback) const
 {
 	return Statements->GetAllActivityDataForEventType(InEventType, [this, InEventType, &InCallback](const int64 InActivityId, const FGuid& InEndpointId, const FDateTime InEventTime, const int64 InEventId, FConcertSessionSerializedPayload&& InEventSummary)
 	{
@@ -1977,13 +1977,13 @@ bool FConcertSyncSessionDatabase::EnumerateActivitiesForEventType(const EConcert
 		Activity.EventType = InEventType;
 		Activity.EventId = InEventId;
 		Activity.EventSummary = MoveTemp(InEventSummary);
-		return InCallback(MoveTemp(Activity))
+		return InCallback(MoveTemp(Activity)) == EBreakBehavior::Continue
 			? ESQLitePreparedStatementExecuteRowResult::Continue
 			: ESQLitePreparedStatementExecuteRowResult::Stop;
 	});
 }
 
-bool FConcertSyncSessionDatabase::EnumerateActivitiesInRange(const int64 InFirstActivityId, const int64 InMaxNumActivities, TFunctionRef<bool(FConcertSyncActivity&&)> InCallback) const
+bool FConcertSyncSessionDatabase::EnumerateActivitiesInRange(const int64 InFirstActivityId, const int64 InMaxNumActivities, FIterateActivityFunc InCallback) const
 {
 	return Statements->GetActivityDataInRange(InFirstActivityId, InMaxNumActivities, [this, &InCallback](const int64 InActivityId, const FGuid& InEndpointId, const FDateTime InEventTime, const EConcertSyncActivityEventType InEventType, const int64 InEventId, FConcertSessionSerializedPayload&& InEventSummary)
 	{
@@ -1995,7 +1995,7 @@ bool FConcertSyncSessionDatabase::EnumerateActivitiesInRange(const int64 InFirst
 		Activity.EventType = InEventType;
 		Activity.EventId = InEventId;
 		Activity.EventSummary = MoveTemp(InEventSummary);
-		return InCallback(MoveTemp(Activity))
+		return InCallback(MoveTemp(Activity)) == EBreakBehavior::Continue
 			? ESQLitePreparedStatementExecuteRowResult::Continue
 			: ESQLitePreparedStatementExecuteRowResult::Stop;
 	});
diff --git a/Engine/Plugins/Developer/Concert/ConcertSync/ConcertSyncCore/Source/ConcertSyncCore/Public/ConcertSyncSessionDatabase.h b/Engine/Plugins/Developer/Concert/ConcertSync/ConcertSyncCore/Source/ConcertSyncCore/Public/ConcertSyncSessionDatabase.h
--- a/Engine/Plugins/Developer/Concert/ConcertSync/ConcertSyncCore/Source/ConcertSyncCore/Public/ConcertSyncSessionDatabase.h
+++ b/Engine/Plugins/Developer/Concert/ConcertSync/ConcertSyncCore/Source/ConcertSyncCore/Public/ConcertSyncSessionDatabase.h
@@ -12,6 +12,18 @@ class FConcertSyncSessionDatabaseStatements;
 class FSQLiteDatabase;
 enum class ESQLiteDatabaseOpenMode : uint8;
 
+enum class EBreakBehavior
+{
+	Break,
+	Continue
+};
+
+using FConsumePackageActivityFunc = TFunctionRef<void(FConcertSyncActivity&&/*BasePart*/, FConcertSyncPackageEventData& /*EventPart*/)>;
+using FIteratePackageActivityFunc = TFunctionRef<EBreakBehavior(FConcertSyncActivity&&/*BasePart*/, FConcertSyncPackageEventData& /*EventPart*/)>;
+using FIterateActivityFunc = TFunctionRef<EBreakBehavior(FConcertSyncActivity&&)>;
+
+DECLARE_MULTICAST_DELEGATE_OneParam(FOnActivityProduced, const FConcertSyncActivity&);
+
 /**
  * Database of activities that have happened in a Concert Sync Session.
  * Stores the activity index and their associated data.
@@ -19,6 +31,7 @@ enum class ESQLiteDatabaseOpenMode : uint8;
 class CONCERTSYNCCORE_API FConcertSyncSessionDatabase
 {
 public:
+	
 	FConcertSyncSessionDatabase();
 	~FConcertSyncSessionDatabase();
 
@@ -220,7 +233,7 @@ public:
 	 *
 	 * @return True if the package activity was found, false otherwise.
 	 */
-	bool GetPackageActivity(const int64 InActivityId, const TFunctionRef<void(FConcertSyncActivity&& /*BasePart*/, FConcertSyncPackageEventData& /*EventPart*/)>& PackageActivityFn) const;
+	bool GetPackageActivity(const int64 InActivityId, FConsumePackageActivityFunc PackageActivityFn) const;
 
 	/**
 	 * Get the type of an activity in this database.
@@ -282,7 +295,7 @@ public:
 	 *
 	 * @return True if the package activity was found, false otherwise.
 	 */
-	bool GetPackageActivityForEvent(const int64 InPackageEventId, const TFunctionRef<void(FConcertSyncActivity&& /*BasePart*/, FConcertSyncPackageEventData& /*EventPart*/)>& PackageActivityFn) const;
+	bool GetPackageActivityForEvent(const int64 InPackageEventId, FIteratePackageActivityFunc PackageActivityFn) const;
 
 	/**
 	 * Enumerate the generic part of the activities in this database.
@@ -291,7 +304,7 @@ public:
 	 *
 	 * @return True if the activities were enumerated without error, false otherwise.
 	 */
-	bool EnumerateActivities(TFunctionRef<bool(FConcertSyncActivity&&)> InCallback) const;
+	bool EnumerateActivities(FIterateActivityFunc InCallback) const;
 
 	/**
 	 * Enumerate all the connection activities in this database.
@@ -327,7 +340,7 @@ public:
 	 *
 	 * @return True if the package activities were enumerated without error, false otherwise.
 	 */
-	bool EnumeratePackageActivities(const TFunctionRef<bool(FConcertSyncActivity&&/*BasePart*/, FConcertSyncPackageEventData& /*EventPart*/)>& InCallback) const;
+	bool EnumeratePackageActivities(FIteratePackageActivityFunc InCallback) const;
 
 	/**
 	 * Enumerate all the activities in this database of the given type.
@@ -337,7 +350,7 @@ public:
 	 *
 	 * @return True if the activities were enumerated without error, false otherwise.
 	 */
-	bool EnumerateActivitiesForEventType(const EConcertSyncActivityEventType InEventType, TFunctionRef<bool(FConcertSyncActivity&&)> InCallback) const;
+	bool EnumerateActivitiesForEventType(const EConcertSyncActivityEventType InEventType, FIterateActivityFunc InCallback) const;
 
 	/**
 	 * Enumerate all the activities in this database in the given range.
@@ -348,7 +361,7 @@ public:
 	 *
 	 * @return True if the activities were enumerated without error, false otherwise.
 	 */
-	bool EnumerateActivitiesInRange(const int64 InFirstActivityId, const int64 InMaxNumActivities, TFunctionRef<bool(FConcertSyncActivity&&)> InCallback) const;
+	bool EnumerateActivitiesInRange(const int64 InFirstActivityId, const int64 InMaxNumActivities, FIterateActivityFunc InCallback) const;
 
 	/**
 	 * Enumerate the IDs and event types of all the activities in this database.
@@ -665,6 +678,8 @@ public:
 	 */
 	void FlushAsynchronousTasks();
 
+	FOnActivityProduced& OnActivityProduced() { return ActivityProducedEvent; }
+	
 private:
 	/**
 	 * Schedule an asynchronous write for the given Package Stream.  The stream must be in-memory. File sharing
@@ -918,7 +933,10 @@ private:
 	 * @return True if the package data was loaded, false otherwise.
 	 */
 	bool LoadPackage(const FString& InPackageBlobFilename, const TFunctionRef<void(FConcertPackageDataStream&)>& PackageDataStreamFn) const;
-
+	
+	/** Called when an activity is produced */
+	FOnActivityProduced ActivityProducedEvent;
+	
 	/** Root path to store all session data under */
 	FString SessionPath;
 
diff --git a/Engine/Plugins/Developer/Concert/ConcertSync/ConcertSyncServer/Source/ConcertSyncServer/Private/ConcertServerWorkspace.cpp b/Engine/Plugins/Developer/Concert/ConcertSync/ConcertSyncServer/Source/ConcertSyncServer/Private/ConcertServerWorkspace.cpp
--- a/Engine/Plugins/Developer/Concert/ConcertSync/ConcertSyncServer/Source/ConcertSyncServer/Private/ConcertServerWorkspace.cpp
+++ b/Engine/Plugins/Developer/Concert/ConcertSync/ConcertSyncServer/Source/ConcertSyncServer/Private/ConcertServerWorkspace.cpp
@@ -1060,6 +1060,8 @@ void FConcertServerWorkspace::PostActivityAdded(const int64 InActivityId)
 	FConcertSyncActivity Activity;
 	if (LiveSession->GetSessionDatabase().GetActivity(InActivityId, Activity))
 	{
+		LiveSession->GetSessionDatabase().OnActivityProduced().Broadcast(Activity);
+		
 		FConcertSyncEndpointData EndpointData;
 		if (LiveSession->GetSessionDatabase().GetEndpoint(Activity.EndpointId, EndpointData))
 		{
diff --git a/Engine/Plugins/Developer/Concert/ConcertUI/ConcertServerUI/Source/ConcertServerUI/Private/Widgets/Session/ConcertServerSessionTab.cpp b/Engine/Plugins/Developer/Concert/ConcertUI/ConcertServerUI/Source/ConcertServerUI/Private/Widgets/Session/ConcertServerSessionTab.cpp
--- a/Engine/Plugins/Developer/Concert/ConcertUI/ConcertServerUI/Source/ConcertServerUI/Private/Widgets/Session/ConcertServerSessionTab.cpp
+++ b/Engine/Plugins/Developer/Concert/ConcertUI/ConcertServerUI/Source/ConcertServerUI/Private/Widgets/Session/ConcertServerSessionTab.cpp
@@ -1,17 +1,19 @@
 ﻿// Copyright Epic Games, Inc. All Rights Reserved.
 
 #include "ConcertServerSessionTab.h"
-
-#include "Framework/Docking/TabManager.h"
 #include "IConcertSession.h"
 #include "Widgets/ConcertServerTabs.h"
-#include "Widgets/Docking/SDockTab.h"
-#include "Widgets/StatusBar/SConcertStatusBar.h"
+#include "Widgets/Session/PackageViewer/ConcertSessionPackageViewerController.h"
 #include "Widgets/Session/SConcertSessionInspector.h"
+#include "Widgets/StatusBar/SConcertStatusBar.h"
+
+#include "Framework/Docking/TabManager.h"
+#include "Widgets/Docking/SDockTab.h"
 
 FConcertServerSessionTab::FConcertServerSessionTab(TSharedRef<IConcertServerSession> InspectedSession, TSharedRef<IConcertSyncServer> SyncServer, const TSharedRef<SWindow>& ConstructUnderWindow)
 	: InspectedSession(MoveTemp(InspectedSession))
 	, SessionHistoryController(MakeShared<FServerSessionHistoryController>(InspectedSession, SyncServer))
+	, PackageViewerController(MakeShared<FConcertSessionPackageViewerController>(InspectedSession, SyncServer))
 	, DockTab(CreateTab(ConstructUnderWindow))
 {}
 
@@ -27,6 +29,9 @@ void FConcertServerSessionTab::OpenSessionTab() const
 	{
 		const FTabManager::FLastMajorOrNomadTab Search(ConcertServerTabs::GetSessionBrowserTabId());
 		FGlobalTabmanager::Get()->InsertNewDocumentTab(*GetTabPlayerHolderId(InspectedSession), Search, DockTab);
+
+		SessionHistoryController->ReloadActivities();
+		PackageViewerController->ReloadActivities();
 	}
 }
 
@@ -36,8 +41,16 @@ TSharedRef<SDockTab> FConcertServerSessionTab::CreateTab(const TSharedRef<SWindo
 	TSharedRef<SDockTab> NewDockTab = SNew(SDockTab)
 		.Label(Title)
 		.TabRole(MajorTab);
+
+	
+	const SConcertSessionInspector::FRequiredArgs WidgetArgs
+	{
+		NewDockTab, ConstructUnderWindow,
+		SessionHistoryController->GetSessionHistory(),
+		PackageViewerController->GetPackageViewer()
+	};
 	NewDockTab->SetContent(
-		SNew(SConcertSessionInspector, SConcertSessionInspector::FRequiredArgs{NewDockTab, ConstructUnderWindow, SessionHistoryController})
+		SNew(SConcertSessionInspector, WidgetArgs)
 			.StatusBar()
 			[
 				SNew(SConcertStatusBar, *GetTabPlayerHolderId(InspectedSession))
diff --git a/Engine/Plugins/Developer/Concert/ConcertUI/ConcertServerUI/Source/ConcertServerUI/Private/Widgets/Session/ConcertServerSessionTab.h b/Engine/Plugins/Developer/Concert/ConcertUI/ConcertServerUI/Source/ConcertServerUI/Private/Widgets/Session/ConcertServerSessionTab.h
--- a/Engine/Plugins/Developer/Concert/ConcertUI/ConcertServerUI/Source/ConcertServerUI/Private/Widgets/Session/ConcertServerSessionTab.h
+++ b/Engine/Plugins/Developer/Concert/ConcertUI/ConcertServerUI/Source/ConcertServerUI/Private/Widgets/Session/ConcertServerSessionTab.h
@@ -5,6 +5,7 @@
 #include "CoreMinimal.h"
 #include "ServerSessionHistoryController.h"
 
+class FConcertSessionPackageViewerController;
 class IConcertSyncServer;
 class FServerSessionHistoryController;
 class IConcertServerSession;
@@ -28,8 +29,11 @@ private:
 
 	/** The session being inspected */
 	const TSharedRef<IConcertServerSession> InspectedSession;
-	
+
+	/** Manages the session history widget */
 	const TSharedRef<FServerSessionHistoryController> SessionHistoryController;
+	/** Manages the package viewer widget */
+	const TSharedRef<FConcertSessionPackageViewerController> PackageViewerController;
 	
 	/** The tab containing the UI for InspectedSession */
 	const TSharedRef<SDockTab> DockTab;
diff --git a/Engine/Plugins/Developer/Concert/ConcertUI/ConcertServerUI/Source/ConcertServerUI/Private/Widgets/Session/PackageViewer/ConcertSessionPackageViewerController.cpp b/Engine/Plugins/Developer/Concert/ConcertUI/ConcertServerUI/Source/ConcertServerUI/Private/Widgets/Session/PackageViewer/ConcertSessionPackageViewerController.cpp
new file mode 100644
--- /dev/null
+++ b/Engine/Plugins/Developer/Concert/ConcertUI/ConcertServerUI/Source/ConcertServerUI/Private/Widgets/Session/PackageViewer/ConcertSessionPackageViewerController.cpp
@@ -0,0 +1,104 @@
+﻿// Copyright Epic Games, Inc. All Rights Reserved.
+
+#include "ConcertSessionPackageViewerController.h"
+
+#include "ConcertSyncSessionDatabase.h"
+#include "IConcertSyncServer.h"
+#include "SConcertSessionPackageViewer.h"
+#include "Concert/Private/ConcertServerSession.h"
+
+FConcertSessionPackageViewerController::FConcertSessionPackageViewerController(TSharedRef<IConcertServerSession> InspectedSession, TSharedRef<IConcertSyncServer> SyncServer)
+	: InspectedSession(MoveTemp(InspectedSession))
+	, SyncServer(MoveTemp(SyncServer))
+	, PackageViewer(MakePackageViewer())
+{
+	ReloadActivities();
+	
+	if (const TOptional<FConcertSyncSessionDatabaseNonNullPtr> Database = SyncServer->GetLiveSessionDatabase(InspectedSession->GetId()))
+	{
+		Database->OnActivityProduced().AddRaw(this, &FConcertSessionPackageViewerController::OnSessionProduced);
+	}
+}
+
+FConcertSessionPackageViewerController::~FConcertSessionPackageViewerController()
+{
+	if (const TOptional<FConcertSyncSessionDatabaseNonNullPtr> Database = SyncServer->GetLiveSessionDatabase(InspectedSession->GetId()))
+	{
+		Database->OnActivityProduced().RemoveAll(this);
+	}
+}
+
+void FConcertSessionPackageViewerController::ReloadActivities() const
+{
+	PackageViewer->ResetActivityList();
+	
+	if (const TOptional<FConcertSyncSessionDatabaseNonNullPtr> Database = SyncServer->GetLiveSessionDatabase(InspectedSession->GetId()))
+	{
+		struct FPackageActivity
+		{
+			FConcertSessionActivity Activity;
+			FConcertSyncPackageEventData PackageData;
+		};
+		
+		TMap<FName, FPackageActivity> LatestPackageActivities;
+		Database->EnumeratePackageActivities([this, &LatestPackageActivities](FConcertSyncActivity&& BasePart, FConcertSyncPackageEventData& EventData)
+		{
+			FStructOnScope ActivitySummary;
+			if (BasePart.EventSummary.GetPayload(ActivitySummary))
+			{
+				FPackageActivity Activity{ { BasePart, ActivitySummary }, EventData };
+				const bool bWasRenamed = Activity.PackageData.MetaData.PackageInfo.PackageUpdateType == EConcertPackageUpdateType::Renamed;
+				if (bWasRenamed)
+				{
+					LatestPackageActivities.Remove(EventData.MetaData.PackageInfo.PackageName);
+					LatestPackageActivities.Emplace(EventData.MetaData.PackageInfo.NewPackageName, MoveTemp(Activity));
+				}
+				else
+				{
+					LatestPackageActivities.Emplace(EventData.MetaData.PackageInfo.PackageName, MoveTemp(Activity));
+				}
+			}
+			return EBreakBehavior::Continue;
+		});
+
+		for (auto ActivityIt = LatestPackageActivities.CreateIterator(); ActivityIt; ++ActivityIt)
+		{
+			PackageViewer->AppendActivity(MoveTemp(ActivityIt->Value.Activity));
+		}
+	}
+}
+
+TSharedRef<SConcertSessionPackageViewer> FConcertSessionPackageViewerController::MakePackageViewer() const
+{
+	return SNew(SConcertSessionPackageViewer)
+		.GetClientInfo_Raw(this, &FConcertSessionPackageViewerController::GetClientInfo)
+		.GetPackageEvent_Raw(this, &FConcertSessionPackageViewerController::GetPackageEvent);
+}
+
+TOptional<FConcertClientInfo> FConcertSessionPackageViewerController::GetClientInfo(FGuid ClientId) const
+{
+	const TOptional<FConcertSyncSessionDatabaseNonNullPtr> Database = SyncServer->GetLiveSessionDatabase(InspectedSession->GetId());
+	FConcertSyncEndpointData Result;
+	if (Database && Database->GetEndpoint(ClientId, Result))
+	{
+		return Result.ClientInfo;
+	}
+	return {};
+}
+
+bool FConcertSessionPackageViewerController::GetPackageEvent(const FConcertSessionActivity& Activity, FConcertSyncPackageEventMetaData& OutPackageEvent) const
+{
+	if (const TOptional<FConcertSyncSessionDatabaseNonNullPtr> Database = SyncServer->GetLiveSessionDatabase(InspectedSession->GetId()))
+	{
+		return Database->GetPackageEventMetaData(Activity.Activity.EventId, OutPackageEvent.PackageRevision, OutPackageEvent.PackageInfo);
+	}
+	return false;
+}
+
+void FConcertSessionPackageViewerController::OnSessionProduced(const FConcertSyncActivity& ProducedActivity) const
+{
+	if (ProducedActivity.EventType == EConcertSyncActivityEventType::Package)
+	{
+		ReloadActivities();
+	}
+}
diff --git a/Engine/Plugins/Developer/Concert/ConcertUI/ConcertServerUI/Source/ConcertServerUI/Private/Widgets/Session/PackageViewer/ConcertSessionPackageViewerController.h b/Engine/Plugins/Developer/Concert/ConcertUI/ConcertServerUI/Source/ConcertServerUI/Private/Widgets/Session/PackageViewer/ConcertSessionPackageViewerController.h
new file mode 100644
--- /dev/null
+++ b/Engine/Plugins/Developer/Concert/ConcertUI/ConcertServerUI/Source/ConcertServerUI/Private/Widgets/Session/PackageViewer/ConcertSessionPackageViewerController.h
@@ -0,0 +1,39 @@
+﻿// Copyright Epic Games, Inc. All Rights Reserved.
+
+#pragma once
+
+#include "CoreMinimal.h"
+
+struct FConcertSyncActivity;
+struct FConcertClientInfo;
+struct FConcertSyncPackageEventMetaData;
+struct FConcertSessionActivity;
+class IConcertSyncServer;
+class IConcertServerSession;
+class SConcertSessionPackageViewer;
+
+/** Manages SConcertSessionPackageViewer for a server. */
+class FConcertSessionPackageViewerController
+{
+public:
+
+	FConcertSessionPackageViewerController(TSharedRef<IConcertServerSession> InspectedSession, TSharedRef<IConcertSyncServer> SyncServer);
+	~FConcertSessionPackageViewerController();
+	
+	void ReloadActivities() const;
+	TSharedRef<SConcertSessionPackageViewer> GetPackageViewer() const { return PackageViewer; }
+
+private:
+	
+	TSharedRef<IConcertServerSession> InspectedSession;
+	TSharedRef<IConcertSyncServer> SyncServer;
+	
+	/** The widget being managed */
+	TSharedRef<SConcertSessionPackageViewer> PackageViewer;
+	
+	TSharedRef<SConcertSessionPackageViewer> MakePackageViewer() const;
+	
+	TOptional<FConcertClientInfo> GetClientInfo(FGuid ClientId) const;
+	bool GetPackageEvent(const FConcertSessionActivity& Activity, FConcertSyncPackageEventMetaData& OutPackageEvent) const;
+	void OnSessionProduced(const FConcertSyncActivity& ProducedActivity) const;
+};
diff --git a/Engine/Plugins/Developer/Concert/ConcertUI/ConcertServerUI/Source/ConcertServerUI/Private/Widgets/Session/PackageViewer/SConcertSessionPackageViewer.cpp b/Engine/Plugins/Developer/Concert/ConcertUI/ConcertServerUI/Source/ConcertServerUI/Private/Widgets/Session/PackageViewer/SConcertSessionPackageViewer.cpp
new file mode 100644
--- /dev/null
+++ b/Engine/Plugins/Developer/Concert/ConcertUI/ConcertServerUI/Source/ConcertServerUI/Private/Widgets/Session/PackageViewer/SConcertSessionPackageViewer.cpp
@@ -0,0 +1,95 @@
+﻿// Copyright Epic Games, Inc. All Rights Reserved.
+
+#include "SConcertSessionPackageViewer.h"
+
+#include "SConcertSessionActivities.h"
+#include "Widgets/Input/SSearchBox.h"
+
+#define LOCTEXT_NAMESPACE "UnrealMultiUserUI"
+
+void SConcertSessionPackageViewer::Construct(const FArguments& InArgs)
+{
+	ActivityListViewOptions = MakeShared<FConcertSessionActivitiesOptions>();
+	ActivityListViewOptions->bEnableConnectionActivityFiltering = false;
+	ActivityListViewOptions->bEnableLockActivityFiltering = false;
+	ActivityListViewOptions->bEnablePackageActivityFiltering = false;
+	ActivityListViewOptions->bEnableTransactionActivityFiltering = false;
+	
+	SAssignNew(ActivityListView, SConcertSessionActivities)
+		.OnGetPackageEvent(InArgs._GetPackageEvent)
+		.OnFetchActivities(InArgs._FetchInitialActivities)
+		.OnMapActivityToClient(InArgs._GetClientInfo)
+		.HighlightText(this, &SConcertSessionPackageViewer::HighlightSearchedText)
+		.TimeFormat(ActivityListViewOptions.Get(), &FConcertSessionActivitiesOptions::GetTimeFormat)
+		.ClientNameColumnVisibility(EVisibility::Visible)
+		.ClientAvatarColorColumnVisibility(EVisibility::Collapsed)
+		.OperationColumnVisibility(EVisibility::Collapsed)
+		.PackageColumnVisibility(EVisibility::Collapsed)
+		.ConnectionActivitiesVisibility(EVisibility::Collapsed)
+		.LockActivitiesVisibility(EVisibility::Collapsed)
+		.PackageActivitiesVisibility(EVisibility::Visible)
+		.TransactionActivitiesVisibility(EVisibility::Collapsed)
+		.DetailsAreaVisibility(EVisibility::Collapsed)
+		.IsAutoScrollEnabled(true);
+
+	ChildSlot
+	[
+		SNew(SVerticalBox)
+		+SVerticalBox::Slot()
+		.AutoHeight()
+		.Padding(1, 1)
+		[
+			SAssignNew(SearchBox, SSearchBox)
+			.HintText(LOCTEXT("SearchHint", "Search..."))
+			.OnTextChanged(this, &SConcertSessionPackageViewer::OnSearchTextChanged)
+			.OnTextCommitted(this, &SConcertSessionPackageViewer::OnSearchTextCommitted)
+			.DelayChangeNotificationsWhileTyping(true)
+		]
+
+		+SVerticalBox::Slot()
+		[
+			ActivityListView.ToSharedRef()
+		]
+
+		+SVerticalBox::Slot()
+		.AutoHeight()
+		[
+			ActivityListViewOptions->MakeStatusBar(
+				TAttribute<int32>(ActivityListView.Get(), &SConcertSessionActivities::GetTotalActivityNum),
+				TAttribute<int32>(ActivityListView.Get(), &SConcertSessionActivities::GetDisplayedActivityNum)
+				)
+		]
+	];
+}
+
+void SConcertSessionPackageViewer::ResetActivityList()
+{
+	ActivityListView->ResetActivityList();
+}
+
+void SConcertSessionPackageViewer::AppendActivity(FConcertSessionActivity Activity)
+{
+	const TSharedRef<FConcertSessionActivity> NewActivity = MakeShared<FConcertSessionActivity>(MoveTemp(Activity));
+	ActivityListView->Append(NewActivity);
+}
+
+void SConcertSessionPackageViewer::OnSearchTextChanged(const FText& InSearchText)
+{
+	SearchedText = InSearchText;
+	SearchBox->SetError(ActivityListView->UpdateTextFilter(InSearchText));
+}
+
+void SConcertSessionPackageViewer::OnSearchTextCommitted(const FText& InFilterText, ETextCommit::Type CommitType)
+{
+	if (!InFilterText.EqualTo(SearchedText))
+	{
+		OnSearchTextChanged(InFilterText);
+	}
+}
+
+FText SConcertSessionPackageViewer::HighlightSearchedText() const
+{
+	return SearchedText;
+}
+
+#undef LOCTEXT_NAMESPACE
diff --git a/Engine/Plugins/Developer/Concert/ConcertUI/ConcertServerUI/Source/ConcertServerUI/Private/Widgets/Session/PackageViewer/SConcertSessionPackageViewer.h b/Engine/Plugins/Developer/Concert/ConcertUI/ConcertServerUI/Source/ConcertServerUI/Private/Widgets/Session/PackageViewer/SConcertSessionPackageViewer.h
new file mode 100644
--- /dev/null
+++ b/Engine/Plugins/Developer/Concert/ConcertUI/ConcertServerUI/Source/ConcertServerUI/Private/Widgets/Session/PackageViewer/SConcertSessionPackageViewer.h
@@ -0,0 +1,53 @@
+﻿// Copyright Epic Games, Inc. All Rights Reserved.
+
+#pragma once
+
+#include "CoreMinimal.h"
+#include "SConcertSessionActivities.h"
+#include "Widgets/SCompoundWidget.h"
+#include "Widgets/DeclarativeSyntaxSupport.h"
+
+class FConcertSessionActivitiesOptions;
+
+/** Displays package activities that happened in a concert session. */
+class SConcertSessionPackageViewer : public SCompoundWidget
+{
+public:
+
+	struct FConcertPackageActivity
+	{
+		FConcertSyncActivity Base;
+		FConcertSyncPackageEventData PackageEvent;
+	};
+	
+	SLATE_BEGIN_ARGS(SConcertSessionPackageViewer) {}
+		SLATE_EVENT(SConcertSessionActivities::FGetPackageEvent, GetPackageEvent)
+		SLATE_EVENT(SConcertSessionActivities::FGetActivityClientInfoFunc, GetClientInfo)
+		SLATE_EVENT(SConcertSessionActivities::FFetchActivitiesFunc, FetchInitialActivities)
+	SLATE_END_ARGS()
+
+	void Construct(const FArguments& InArgs);
+
+	void ResetActivityList();
+	void AppendActivity(FConcertSessionActivity Activity);
+
+private:
+
+	TSharedPtr<SConcertSessionActivities> ActivityListView;
+	
+	/** Controls the activity list view options */
+	TSharedPtr<FConcertSessionActivitiesOptions> ActivityListViewOptions;
+
+	/** The widget used to enter the text to search. */
+	TSharedPtr<SSearchBox> SearchBox;
+
+	/** The searched text to highlight. */
+	FText SearchedText;
+
+	/** Invoked when the text in the search box widget changes. */
+	void OnSearchTextChanged(const FText& InSearchText);
+	/** Invoked when the text in the search box widget is committed. */
+	void OnSearchTextCommitted(const FText& InFilterText, ETextCommit::Type CommitType);
+	/** Returns the text to highlight when the search bar has a text set. */
+	FText HighlightSearchedText() const;
+};
diff --git a/Engine/Plugins/Developer/Concert/ConcertUI/ConcertServerUI/Source/ConcertServerUI/Private/Widgets/Session/SConcertSessionInspector.cpp b/Engine/Plugins/Developer/Concert/ConcertUI/ConcertServerUI/Source/ConcertServerUI/Private/Widgets/Session/SConcertSessionInspector.cpp
--- a/Engine/Plugins/Developer/Concert/ConcertUI/ConcertServerUI/Source/ConcertServerUI/Private/Widgets/Session/SConcertSessionInspector.cpp
+++ b/Engine/Plugins/Developer/Concert/ConcertUI/ConcertServerUI/Source/ConcertServerUI/Private/Widgets/Session/SConcertSessionInspector.cpp
@@ -57,10 +57,10 @@ TSharedRef<SWidget> SConcertSessionInspector::CreateTabs(const FArguments& InArg
 {
 	TabManager = FGlobalTabmanager::Get()->NewTabManager(RequiredArgs.ConstructUnderMajorTab);
 	
-	TabManager->RegisterTabSpawner(HistoryTabId, FOnSpawnTab::CreateSP(this, &SConcertSessionInspector::SpawnActivityHistory, RequiredArgs.SessionHistoryController))
+	TabManager->RegisterTabSpawner(HistoryTabId, FOnSpawnTab::CreateSP(this, &SConcertSessionInspector::SpawnActivityHistory, RequiredArgs.SessionHistory))
 		.SetDisplayName(LOCTEXT("ActivityHistoryLabel", "History"));
 	
-	TabManager->RegisterTabSpawner(SessionContentTabId, FOnSpawnTab::CreateSP(this, &SConcertSessionInspector::SpawnSessionContent))
+	TabManager->RegisterTabSpawner(SessionContentTabId, FOnSpawnTab::CreateSP(this, &SConcertSessionInspector::SpawnSessionContent, RequiredArgs.PackageViewer))
 		.SetDisplayName(LOCTEXT("SessionContentLabel", "Session Content"));
 	
 	TabManager->RegisterTabSpawner(ConnectionMonitorTabId, FOnSpawnTab::CreateSP(this, &SConcertSessionInspector::SpawnConnectionMonitor))
@@ -100,23 +100,23 @@ TSharedRef<SWidget> SConcertSessionInspector::CreateTabs(const FArguments& InArg
 	return TabManager->RestoreFrom(Layout, RequiredArgs.ConstructUnderWindow).ToSharedRef();
 }
 
-TSharedRef<SDockTab> SConcertSessionInspector::SpawnActivityHistory(const FSpawnTabArgs& Args, TSharedRef<FServerSessionHistoryController> SessionHistoryController)
+TSharedRef<SDockTab> SConcertSessionInspector::SpawnActivityHistory(const FSpawnTabArgs& Args, TSharedRef<SSessionHistory> SessionHistory)
 {
 	return SNew(SDockTab)
 		.Label(LOCTEXT("ActivityHistoryLabel", "History"))
 		.TabRole(PanelTab)
 		[
-			SessionHistoryController->GetSessionHistory()
+			SessionHistory
 		];
 }
 
-TSharedRef<SDockTab> SConcertSessionInspector::SpawnSessionContent(const FSpawnTabArgs& Args)
+TSharedRef<SDockTab> SConcertSessionInspector::SpawnSessionContent(const FSpawnTabArgs& Args, TSharedRef<SConcertSessionPackageViewer> PackageViewer)
 {
 	return SNew(SDockTab)
 		.Label(LOCTEXT("SessionContentLabel", "Session Content"))
 		.TabRole(PanelTab)
 		[
-			SNullWidget::NullWidget
+			PackageViewer
 		];
 }
 
diff --git a/Engine/Plugins/Developer/Concert/ConcertUI/ConcertServerUI/Source/ConcertServerUI/Private/Widgets/Session/SConcertSessionInspector.h b/Engine/Plugins/Developer/Concert/ConcertUI/ConcertServerUI/Source/ConcertServerUI/Private/Widgets/Session/SConcertSessionInspector.h
--- a/Engine/Plugins/Developer/Concert/ConcertUI/ConcertServerUI/Source/ConcertServerUI/Private/Widgets/Session/SConcertSessionInspector.h
+++ b/Engine/Plugins/Developer/Concert/ConcertUI/ConcertServerUI/Source/ConcertServerUI/Private/Widgets/Session/SConcertSessionInspector.h
@@ -3,7 +3,8 @@
 #pragma once
 
 #include "CoreMinimal.h"
-#include "ServerSessionHistoryController.h"
+#include "SessionHistory/SSessionHistory.h"
+#include "PackageViewer/SConcertSessionPackageViewer.h"
 #include "Widgets/SCompoundWidget.h"
 #include "Widgets/DeclarativeSyntaxSupport.h"
 
@@ -27,12 +28,16 @@ public:
 	{
 		TSharedRef<SDockTab> ConstructUnderMajorTab;
 		TSharedRef<SWindow> ConstructUnderWindow;
-		TSharedRef<FServerSessionHistoryController> SessionHistoryController;
+		TSharedRef<SSessionHistory> SessionHistory;
+		TSharedRef<SConcertSessionPackageViewer> PackageViewer;
 
-		FRequiredArgs(TSharedRef<SDockTab> ConstructUnderMajorTab, TSharedRef<SWindow> ConstructUnderWindow,TSharedRef<FServerSessionHistoryController> SessionHistoryController)
+		FRequiredArgs(TSharedRef<SDockTab> ConstructUnderMajorTab, TSharedRef<SWindow> ConstructUnderWindow,
+				TSharedRef<SSessionHistory> SessionHistoryController,
+				TSharedRef<SConcertSessionPackageViewer> PackageViewerController)
 			: ConstructUnderMajorTab(MoveTemp(ConstructUnderMajorTab))
 			, ConstructUnderWindow(MoveTemp(ConstructUnderWindow))
-			, SessionHistoryController(MoveTemp(SessionHistoryController))
+			, SessionHistory(MoveTemp(SessionHistoryController))
+			, PackageViewer(MoveTemp(PackageViewerController))
 		{}
 	};
 
@@ -53,7 +58,7 @@ private:
 
 	// Spawn tabs
 	TSharedRef<SWidget> CreateTabs(const FArguments& InArgs, const FRequiredArgs& RequiredArgs);
-	TSharedRef<SDockTab> SpawnActivityHistory(const FSpawnTabArgs& Args, TSharedRef<FServerSessionHistoryController> SessionHistoryController);
-	TSharedRef<SDockTab> SpawnSessionContent(const FSpawnTabArgs& Args);
+	TSharedRef<SDockTab> SpawnActivityHistory(const FSpawnTabArgs& Args, TSharedRef<SSessionHistory> SessionHistory);
+	TSharedRef<SDockTab> SpawnSessionContent(const FSpawnTabArgs& Args, TSharedRef<SConcertSessionPackageViewer> PackageViewer);
 	TSharedRef<SDockTab> SpawnConnectionMonitor(const FSpawnTabArgs& Args);
 };
diff --git a/Engine/Plugins/Developer/Concert/ConcertUI/ConcertServerUI/Source/ConcertServerUI/Private/Widgets/Session/ServerSessionHistoryController.cpp b/Engine/Plugins/Developer/Concert/ConcertUI/ConcertServerUI/Source/ConcertServerUI/Private/Widgets/Session/ServerSessionHistoryController.cpp
--- a/Engine/Plugins/Developer/Concert/ConcertUI/ConcertServerUI/Source/ConcertServerUI/Private/Widgets/Session/ServerSessionHistoryController.cpp
+++ b/Engine/Plugins/Developer/Concert/ConcertUI/ConcertServerUI/Source/ConcertServerUI/Private/Widgets/Session/ServerSessionHistoryController.cpp
@@ -9,7 +9,22 @@
 FServerSessionHistoryController::FServerSessionHistoryController(TSharedRef<IConcertServerSession> InspectedSession, TSharedRef<IConcertSyncServer> SyncServer)
 	: InspectedSession(MoveTemp(InspectedSession))
 	, SyncServer(MoveTemp(SyncServer))
-{}
+{
+	ReloadActivities();
+
+	if (const TOptional<FConcertSyncSessionDatabaseNonNullPtr> Database = SyncServer->GetLiveSessionDatabase(InspectedSession->GetId()))
+	{
+		Database->OnActivityProduced().AddRaw(this, &FServerSessionHistoryController::OnSessionProduced);
+	}
+}
+
+FServerSessionHistoryController::~FServerSessionHistoryController()
+{
+	if (const TOptional<FConcertSyncSessionDatabaseNonNullPtr> Database = SyncServer->GetLiveSessionDatabase(InspectedSession->GetId()))
+	{
+		Database->OnActivityProduced().RemoveAll(this);
+	}
+}
 
 void FServerSessionHistoryController::GetActivities(int64 MaximumNumberOfActivities, TMap<FGuid, FConcertClientInfo>& OutEndpointClientInfoMap, TArray<FConcertSessionActivity>& OutFetchedActivities) const
 {
@@ -38,7 +53,7 @@ void FServerSessionHistoryController::GetActivities(int64 MaximumNumberOfActivit
 				OutFetchedActivities.Emplace(MoveTemp(InActivity), MoveTemp(ActivitySummary));
 			}
 
-			return true;
+			return EBreakBehavior::Continue;
 		});
 	}
 }
@@ -70,3 +85,8 @@ TFuture<TOptional<FConcertSyncTransactionEvent>> FServerSessionHistoryController
 		?	MakeFulfilledPromise<TOptional<FConcertSyncTransactionEvent>>(MoveTemp(TransactionEvent)).GetFuture()
 		:	MakeFulfilledPromise<TOptional<FConcertSyncTransactionEvent>>().GetFuture();
 }
+
+void FServerSessionHistoryController::OnSessionProduced(const FConcertSyncActivity& ProducedActivity)
+{
+	ReloadActivities();
+}
diff --git a/Engine/Plugins/Developer/Concert/ConcertUI/ConcertServerUI/Source/ConcertServerUI/Private/Widgets/Session/ServerSessionHistoryController.h b/Engine/Plugins/Developer/Concert/ConcertUI/ConcertServerUI/Source/ConcertServerUI/Private/Widgets/Session/ServerSessionHistoryController.h
--- a/Engine/Plugins/Developer/Concert/ConcertUI/ConcertServerUI/Source/ConcertServerUI/Private/Widgets/Session/ServerSessionHistoryController.h
+++ b/Engine/Plugins/Developer/Concert/ConcertUI/ConcertServerUI/Source/ConcertServerUI/Private/Widgets/Session/ServerSessionHistoryController.h
@@ -5,6 +5,7 @@
 #include "CoreMinimal.h"
 #include "SessionHistory/AbstractSessionHistoryController.h"
 
+struct FConcertSyncActivity;
 class FConcertSyncSessionDatabase;
 class IConcertSyncServer;
 class IConcertServerSession;
@@ -14,6 +15,7 @@ class FServerSessionHistoryController : public FAbstractSessionHistoryController
 public:
 
 	FServerSessionHistoryController(TSharedRef<IConcertServerSession> InspectedSession, TSharedRef<IConcertSyncServer> SyncServer);
+	virtual ~FServerSessionHistoryController() override;
 	
 protected:
 
@@ -29,4 +31,6 @@ private:
 	TSharedRef<IConcertSyncServer> SyncServer;
 
 	TFuture<TOptional<FConcertSyncTransactionEvent>> FindOrRequestTransactionEvent(const FConcertSyncSessionDatabase& Database, const int64 TransactionEventId) const;
+	
+	void OnSessionProduced(const FConcertSyncActivity& ProducedActivity);
 };
diff --git a/Engine/Plugins/Developer/Concert/ConcertUI/ConcertSharedSlate/Source/ConcertSharedSlate/Private/SConcertSessionActivities.cpp b/Engine/Plugins/Developer/Concert/ConcertUI/ConcertSharedSlate/Source/ConcertSharedSlate/Private/SConcertSessionActivities.cpp
--- a/Engine/Plugins/Developer/Concert/ConcertUI/ConcertSharedSlate/Source/ConcertSharedSlate/Private/SConcertSessionActivities.cpp
+++ b/Engine/Plugins/Developer/Concert/ConcertUI/ConcertSharedSlate/Source/ConcertSharedSlate/Private/SConcertSessionActivities.cpp
@@ -131,6 +131,11 @@ FText GetSummary(const FConcertSessionActivity& Activity, const FText& ClientNam
 	return FText::GetEmpty();
 }
 
+FText GetClientName(const TOptional<FConcertClientInfo>& InActivityClient)
+{
+	return InActivityClient ? FText::AsCultureInvariant(InActivityClient->DisplayName) : FText::GetEmpty();
+}
+	
 FText GetClientName(const FConcertClientInfo* InActivityClient)
 {
 	return InActivityClient ? FText::AsCultureInvariant(InActivityClient->DisplayName) : FText::GetEmpty();
@@ -165,7 +170,7 @@ public:
 	 * @param InActivityClient The client who produced this activity. Can be null if unknown or not desirable.
 	 * @param InOwnerTableView The table view that will own this row.
 	 */
-	void Construct(const FArguments& InArgs, TSharedPtr<FConcertSessionActivity> InActivity, const FConcertClientInfo* InActivityClient, const TSharedRef<STableViewBase>& InOwnerTableView);
+	void Construct(const FArguments& InArgs, TSharedPtr<FConcertSessionActivity> InActivity, const TOptional<FConcertClientInfo>& InActivityClient, const TSharedRef<STableViewBase>& InOwnerTableView);
 
 	/** Generates the widget representing this row. */
 	virtual TSharedRef<SWidget> GenerateWidgetForColumn(const FName& ColumnName) override;
@@ -194,7 +199,7 @@ private:
 };
 
 
-void SConcertSessionActivityRow::Construct(const FArguments& InArgs, TSharedPtr<FConcertSessionActivity> InActivity, const FConcertClientInfo* InActivityClient, const TSharedRef<STableViewBase>& InOwnerTableView)
+void SConcertSessionActivityRow::Construct(const FArguments& InArgs, TSharedPtr<FConcertSessionActivity> InActivity, const TOptional<FConcertClientInfo>& InActivityClient, const TSharedRef<STableViewBase>& InOwnerTableView)
 {
 	Activity = InActivity;
 	TimeFormat = InArgs._TimeFormat;
@@ -284,9 +289,9 @@ TSharedRef<SWidget> SConcertSessionActivityRow::GenerateWidgetForColumn(const FN
 		];
 	}
 
-	if (OnMakeColumnOverlayWidget)
+	if (OnMakeColumnOverlayWidget.IsBound())
 	{
-		if (TSharedPtr<SWidget> OverlayedWidget = OnMakeColumnOverlayWidget(ActivityPin, ColumnId))
+		if (TSharedPtr<SWidget> OverlayedWidget = OnMakeColumnOverlayWidget.Execute(ActivityPin, ColumnId))
 		{
 			Overlay->AddSlot()
 			[
@@ -572,7 +577,7 @@ void SConcertSessionActivities::OnListViewScrolled(double InScrollOffset)
 {
 	bActivityViewScrolled = true;
 
-	if (FetchActivitiesFn) // This widget is responsible to populate the view.
+	if (FetchActivitiesFn.IsBound()) // This widget is responsible to populate the view.
 	{
 		if (!bAllActivitiesFetched && ActivityView->GetScrollDistance().Y > 0.7) // Should fetch more?
 		{
@@ -629,11 +634,11 @@ void SConcertSessionActivities::UpdateDetailArea(TSharedPtr<FConcertSessionActiv
 			SetDetailsPanelVisibility(NoDetailsPanel.Get());
 		}
 	}
-	else if (InSelectedActivity->Activity.EventType == EConcertSyncActivityEventType::Transaction && GetTransactionEventFn) // A function is bound to get the transaction event?
+	else if (InSelectedActivity->Activity.EventType == EConcertSyncActivityEventType::Transaction && GetTransactionEventFn.IsBound()) // A function is bound to get the transaction event?
 	{
 		SetDetailsPanelVisibility(LoadingDetailsPanel.Get());
 		TWeakPtr<SConcertSessionActivities> WeakSelf = SharedThis(this);
-		GetTransactionEventFn(*InSelectedActivity).Next([WeakSelf, InSelectedActivity](const TOptional<FConcertSyncTransactionEvent>& TransactionEvent)
+		GetTransactionEventFn.Execute(*InSelectedActivity).Next([WeakSelf, InSelectedActivity](const TOptional<FConcertSyncTransactionEvent>& TransactionEvent)
 		{
 			if (TSharedPtr<SConcertSessionActivities> Self = WeakSelf.Pin()) // If 'this' object hasn't been deleted.
 			{
@@ -653,11 +658,11 @@ void SConcertSessionActivities::UpdateDetailArea(TSharedPtr<FConcertSessionActiv
 			// else -> The widget was deleted.
 		});
 	}
-	else if (InSelectedActivity->Activity.EventType == EConcertSyncActivityEventType::Package && GetPackageEventFn) // A function is bound to get the package event?
+	else if (InSelectedActivity->Activity.EventType == EConcertSyncActivityEventType::Package && GetPackageEventFn.IsBound()) // A function is bound to get the package event?
 	{
 		SetDetailsPanelVisibility(LoadingDetailsPanel.Get());
 		FConcertSyncPackageEventMetaData PackageEventMetaData;
-		if (GetPackageEventFn(*InSelectedActivity, PackageEventMetaData))
+		if (GetPackageEventFn.Execute(*InSelectedActivity, PackageEventMetaData))
 		{
 			DisplayPackageDetails(*InSelectedActivity, PackageEventMetaData.PackageRevision, PackageEventMetaData.PackageInfo);
 		}
@@ -734,7 +739,7 @@ void SConcertSessionActivities::OnActivityFilterUpdated()
 
 void SConcertSessionActivities::FetchActivities()
 {
-	if (!FetchActivitiesFn) // Not bound?
+	if (!FetchActivitiesFn.IsBound()) // Not bound?
 	{
 		return; // The widget is expected to be populated/cleared externally using Append()/Reset()
 	}
@@ -748,7 +753,7 @@ void SConcertSessionActivities::FetchActivities()
 		int32 FetchCount = 0; // The number of activities fetched in this iteration.
 		int32 StartInsertPos = AllActivities.Num();
 
-		bAllActivitiesFetched = FetchActivitiesFn(AllActivities, FetchCount, ErrorMsg);
+		bAllActivitiesFetched = FetchActivitiesFn.Execute(AllActivities, FetchCount, ErrorMsg);
 		if (ErrorMsg.IsEmpty())
 		{
 			if (FetchCount) // New activities appended?
@@ -822,7 +827,7 @@ void SConcertSessionActivities::RequestRefresh()
 	ActivityView->RequestListRefresh();
 }
 
-void SConcertSessionActivities::Reset()
+void SConcertSessionActivities::ResetActivityList()
 {
 	Activities.Reset();
 	AllActivities.Reset();
@@ -867,7 +872,8 @@ FText SConcertSessionActivities::UpdateTextFilter(const FText& InFilterText)
 
 void SConcertSessionActivities::PopulateSearchStrings(const FConcertSessionActivity& Activity, TArray<FString>& OutSearchStrings) const
 {
-	FText ClientName = GetActivityUserFn ? ConcertSessionActivityUtils::GetClientName(GetActivityUserFn(Activity.Activity.EndpointId)) : FText::GetEmpty();
+	FText ClientName = GetActivityUserFn.IsBound()
+		? ConcertSessionActivityUtils::GetClientName(GetActivityUserFn.Execute(Activity.Activity.EndpointId)) : FText::GetEmpty();
 
 	OutSearchStrings.Add(ConcertSessionActivityUtils::GetActivityDateTime(Activity, TimeFormat.Get()).ToString());
 	OutSearchStrings.Add(ConcertSessionActivityUtils::GetSummary(Activity, ClientName, /*bAsRichText*/false).ToString());
@@ -890,7 +896,8 @@ void SConcertSessionActivities::PopulateSearchStrings(const FConcertSessionActiv
 
 TSharedRef<ITableRow> SConcertSessionActivities::OnGenerateActivityRowWidget(TSharedPtr<FConcertSessionActivity> Activity, const TSharedRef<STableViewBase>& OwnerTable)
 {
-	return SNew(SConcertSessionActivityRow, Activity, GetActivityUserFn ? GetActivityUserFn(Activity->Activity.EndpointId) : nullptr, OwnerTable)
+	const TOptional<FConcertClientInfo> ActivityClient = GetActivityUserFn.IsBound() ? GetActivityUserFn.Execute(Activity->Activity.EndpointId) : TOptional<FConcertClientInfo>{};
+	return SNew(SConcertSessionActivityRow, Activity, ActivityClient, OwnerTable)
 		.TimeFormat(TimeFormat)
 		.HighlightText(HighlightText)
 		.OnMakeColumnOverlayWidget(MakeColumnOverlayWidgetFn);
@@ -959,12 +966,8 @@ void SConcertSessionActivities::DisplayTransactionDetails(const FConcertSessionA
 
 void SConcertSessionActivities::DisplayPackageDetails(const FConcertSessionActivity& Activity, int64 PackageRevision, const FConcertPackageInfo& PackageInfo)
 {
-	const FConcertClientInfo* ClientInfo = nullptr;
-	if (GetActivityUserFn)
-	{
-		ClientInfo = GetActivityUserFn(Activity.Activity.EndpointId);
-	}
-
+	const TOptional<FConcertClientInfo> ClientInfo = GetActivityUserFn.IsBound()
+		? GetActivityUserFn.Execute(Activity.Activity.EndpointId) : TOptional<FConcertClientInfo>{};
 	PackageDetailsPanel->SetPackageInfo(PackageInfo, PackageRevision, ClientInfo ? ClientInfo->DisplayName : FString());
 	SetDetailsPanelVisibility(PackageDetailsPanel.Get());
 }
diff --git a/Engine/Plugins/Developer/Concert/ConcertUI/ConcertSharedSlate/Source/ConcertSharedSlate/Private/SConcertSessionRecovery.cpp b/Engine/Plugins/Developer/Concert/ConcertUI/ConcertSharedSlate/Source/ConcertSharedSlate/Private/SConcertSessionRecovery.cpp
--- a/Engine/Plugins/Developer/Concert/ConcertUI/ConcertSharedSlate/Source/ConcertSharedSlate/Private/SConcertSessionRecovery.cpp
+++ b/Engine/Plugins/Developer/Concert/ConcertUI/ConcertSharedSlate/Source/ConcertSharedSlate/Private/SConcertSessionRecovery.cpp
@@ -36,7 +36,7 @@ void SConcertSessionRecovery::Construct(const FArguments& InArgs)
 	SAssignNew(ActivityView, SConcertSessionActivities)
 	.OnFetchActivities(InArgs._OnFetchActivities)
 	.OnMapActivityToClient(InArgs._OnMapActivityToClient)
-	.OnMakeColumnOverlayWidget([this](TWeakPtr<FConcertSessionActivity> Activity, const FName& ColumnId) { return MakeRecoverThroughWidget(Activity, ColumnId); })
+	.OnMakeColumnOverlayWidget_Lambda([this](TWeakPtr<FConcertSessionActivity> Activity, const FName& ColumnId) { return MakeRecoverThroughWidget(Activity, ColumnId); })
 	.HighlightText(this, &SConcertSessionRecovery::HighlightSearchText)
 	.TimeFormat(ActivityViewOptions.Get(), &FConcertSessionActivitiesOptions::GetTimeFormat)
 	.ClientAvatarColorColumnVisibility(InArgs._ClientAvatarColorColumnVisibility)
@@ -172,7 +172,7 @@ TSharedPtr<SWidget> SConcertSessionRecovery::MakeRecoverThroughWidget(TWeakPtr<F
 
 void SConcertSessionRecovery::Reset()
 {
-	ActivityView->Reset();
+	ActivityView->ResetActivityList();
 }
 
 int32 SConcertSessionRecovery::GetTotalActivityNum() const
diff --git a/Engine/Plugins/Developer/Concert/ConcertUI/ConcertSharedSlate/Source/ConcertSharedSlate/Private/SessionHistory/AbstractSessionHistoryController.cpp b/Engine/Plugins/Developer/Concert/ConcertUI/ConcertSharedSlate/Source/ConcertSharedSlate/Private/SessionHistory/AbstractSessionHistoryController.cpp
--- a/Engine/Plugins/Developer/Concert/ConcertUI/ConcertSharedSlate/Source/ConcertSharedSlate/Private/SessionHistory/AbstractSessionHistoryController.cpp
+++ b/Engine/Plugins/Developer/Concert/ConcertUI/ConcertSharedSlate/Source/ConcertSharedSlate/Private/SessionHistory/AbstractSessionHistoryController.cpp
@@ -27,6 +27,6 @@ TSharedRef<SSessionHistory> FAbstractSessionHistoryController::MakeSessionHistor
 {
 	return SNew(SSessionHistory)
 		.PackageFilter(PackageFilter)
-		.GetPackageEvent([this](const FConcertSessionActivity& Activity, FConcertSyncPackageEventMetaData& OutPackageEvent) { return GetPackageEvent(Activity, OutPackageEvent); })
-		.GetTransactionEvent([this](const FConcertSessionActivity& Activity) { return GetTransactionEvent(Activity); });
+		.GetPackageEvent_Raw(this, &FAbstractSessionHistoryController::GetPackageEvent)
+		.GetTransactionEvent_Raw(this, &FAbstractSessionHistoryController::GetTransactionEvent);
 }
diff --git a/Engine/Plugins/Developer/Concert/ConcertUI/ConcertSharedSlate/Source/ConcertSharedSlate/Private/SessionHistory/SSessionHistory.cpp b/Engine/Plugins/Developer/Concert/ConcertUI/ConcertSharedSlate/Source/ConcertSharedSlate/Private/SessionHistory/SSessionHistory.cpp
--- a/Engine/Plugins/Developer/Concert/ConcertUI/ConcertSharedSlate/Source/ConcertSharedSlate/Private/SessionHistory/SSessionHistory.cpp
+++ b/Engine/Plugins/Developer/Concert/ConcertUI/ConcertSharedSlate/Source/ConcertSharedSlate/Private/SessionHistory/SSessionHistory.cpp
@@ -49,7 +49,7 @@ void SSessionHistory::Construct(const FArguments& InArgs)
 	SAssignNew(ActivityListView, SConcertSessionActivities)
 		.OnGetPackageEvent(InArgs._GetPackageEvent)
 		.OnGetTransactionEvent(InArgs._GetTransactionEvent)
-		.OnMapActivityToClient([this](FGuid ClientId){ return EndpointClientInfoMap.Find(ClientId); })
+		.OnMapActivityToClient(this, &SSessionHistory::GetClientInfo)
 		.HighlightText(this, &SSessionHistory::HighlightSearchedText)
 		.TimeFormat(ActivityListViewOptions.Get(), &FConcertSessionActivitiesOptions::GetTimeFormat)
 		.ClientNameColumnVisibility(EVisibility::Visible)
@@ -104,7 +104,7 @@ void SSessionHistory::ReloadActivities(TMap<FGuid, FConcertClientInfo> InEndpoin
 {
 	EndpointClientInfoMap = MoveTemp(InEndpointClientInfoMap);
 	ActivityMap.Reset();
-	ActivityListView->Reset(); // Careful, don't reset the shared ptr.
+	ActivityListView->ResetActivityList(); 
 
 	for (FConcertSessionActivity& FetchedActivity : InFetchedActivities)
 	{
@@ -163,4 +163,13 @@ FText SSessionHistory::HighlightSearchedText() const
 	return SearchedText;
 }
 
+TOptional<FConcertClientInfo> SSessionHistory::GetClientInfo(FGuid Guid) const
+{
+	if (const FConcertClientInfo* ClientInfo = EndpointClientInfoMap.Find(Guid))
+	{
+		return *ClientInfo;
+	}
+	return {};
+}
+
 #undef LOCTEXT_NAMESPACE /* SSessionHistory */
diff --git a/Engine/Plugins/Developer/Concert/ConcertUI/ConcertSharedSlate/Source/ConcertSharedSlate/Public/SConcertSessionActivities.h b/Engine/Plugins/Developer/Concert/ConcertUI/ConcertSharedSlate/Source/ConcertSharedSlate/Public/SConcertSessionActivities.h
--- a/Engine/Plugins/Developer/Concert/ConcertUI/ConcertSharedSlate/Source/ConcertSharedSlate/Public/SConcertSessionActivities.h
+++ b/Engine/Plugins/Developer/Concert/ConcertUI/ConcertSharedSlate/Source/ConcertSharedSlate/Public/SConcertSessionActivities.h
@@ -37,6 +37,7 @@ ENUM_CLASS_FLAGS(EConcertActivityFilterFlags);
 class CONCERTSHAREDSLATE_API SConcertSessionActivities : public SCompoundWidget
 {
 public:
+	
 	/** Defines how the time should be displayed in the date/time column. */
 	enum class ETimeFormat
 	{
@@ -45,19 +46,20 @@ public:
 	};
 
 	/** Used to pull activities from a session. Used to fetch and display the activities of an archived session. */
-	using FFetchActivitiesFunc = TFunction<bool(TArray<TSharedPtr<FConcertSessionActivity>>& /*InOutActivities*/, int32& /*OutFetchedCount*/, FText& /*ErrorMsg*/)>;
-
+	DECLARE_DELEGATE_RetVal_ThreeParams(bool, FFetchActivitiesFunc, TArray<TSharedPtr<FConcertSessionActivity>>& /*InOutActivities*/, int32& /*OutFetchedCount*/, FText& /*ErrorMsg*/);
+	
 	/** Used to map an activity to its client. */
-	using FGetActivityClientInfoFunc = TFunction<const FConcertClientInfo*(FGuid /*ClientId*/)>;
+	DECLARE_DELEGATE_RetVal_OneParam(TOptional<FConcertClientInfo>, FGetActivityClientInfoFunc, FGuid /*ClientId*/);
 
 	/** Returns the transaction event corresponding the specified activity.*/
-	using FGetTransactionEvent = TFunction<TFuture<TOptional<FConcertSyncTransactionEvent>>(const FConcertSessionActivity& /*Activity*/)>;
+	DECLARE_DELEGATE_RetVal_OneParam(TFuture<TOptional<FConcertSyncTransactionEvent>>, FGetTransactionEvent, const FConcertSessionActivity& /*Activity*/)
+	
 
 	/** Returns the package event corresponding to the package activity. */
-	using FGetPackageEvent = TFunction<bool(const FConcertSessionActivity& /*Activity*/, FConcertSyncPackageEventMetaData& /*OutEvent*/)>;
+	DECLARE_DELEGATE_RetVal_TwoParams(bool, FGetPackageEvent, const FConcertSessionActivity& /*Activity*/, FConcertSyncPackageEventMetaData& /*OutEvent*/);
 
 	/** Used to overlay a widget over a column widget to add custom functionalities to a row. */
-	using FMakeColumnOverlayWidgetFunc = TFunction<TSharedPtr<SWidget>(TWeakPtr<FConcertSessionActivity> /*ThisRowActivity*/, const FName& /*ColumnId*/)>;
+	DECLARE_DELEGATE_RetVal_TwoParams(TSharedPtr<SWidget>, FMakeColumnOverlayWidgetFunc, TWeakPtr<FConcertSessionActivity> /*ThisRowActivity*/, const FName& /*ColumnId*/);
 
 public:
 	SLATE_BEGIN_ARGS(SConcertSessionActivities)
@@ -78,19 +80,19 @@ public:
 		, _IsAutoScrollEnabled(false){ }
 
 		/** If bound, invoked to populate the view. */
-		SLATE_ARGUMENT(FFetchActivitiesFunc, OnFetchActivities)
+		SLATE_EVENT(FFetchActivitiesFunc, OnFetchActivities)
 
 		/** If bound, invoked to map an activity to a client.*/
-		SLATE_ARGUMENT(FGetActivityClientInfoFunc, OnMapActivityToClient)
+		SLATE_EVENT(FGetActivityClientInfoFunc, OnMapActivityToClient)
 
 		/** If bound, invoked to fill up the package activity details panel. */
-		SLATE_ARGUMENT(FGetPackageEvent, OnGetPackageEvent)
+		SLATE_EVENT(FGetPackageEvent, OnGetPackageEvent)
 
 		/** If bound, invoked to fill up the transaction activity details panel. */
-		SLATE_ARGUMENT(FGetTransactionEvent, OnGetTransactionEvent)
+		SLATE_EVENT(FGetTransactionEvent, OnGetTransactionEvent)
 
 		/** If bound, invoked when generating a row to add an overlay to a column. */
-		SLATE_ARGUMENT(FMakeColumnOverlayWidgetFunc, OnMakeColumnOverlayWidget)
+		SLATE_EVENT(FMakeColumnOverlayWidgetFunc, OnMakeColumnOverlayWidget)
 
 		/** Highlight the returned text in the view. */
 		SLATE_ATTRIBUTE(FText, HighlightText)
@@ -163,7 +165,7 @@ public:
 	bool IsLastColumn(const FName& ColumnId) const;
 
 	/** Clears all activities displayed. */
-	void Reset();
+	void ResetActivityList();
 
 	/** Append an activity to the view. Used to populate the view from a live session. */
 	void Append(TSharedPtr<FConcertSessionActivity> Activity);
diff --git a/Engine/Plugins/Developer/Concert/ConcertUI/ConcertSharedSlate/Source/ConcertSharedSlate/Public/SConcertSessionRecovery.h b/Engine/Plugins/Developer/Concert/ConcertUI/ConcertSharedSlate/Source/ConcertSharedSlate/Public/SConcertSessionRecovery.h
--- a/Engine/Plugins/Developer/Concert/ConcertUI/ConcertSharedSlate/Source/ConcertSharedSlate/Public/SConcertSessionRecovery.h
+++ b/Engine/Plugins/Developer/Concert/ConcertUI/ConcertSharedSlate/Source/ConcertSharedSlate/Public/SConcertSessionRecovery.h
@@ -3,6 +3,7 @@
 #pragma once
 
 #include "CoreMinimal.h"
+#include "SConcertSessionActivities.h"
 #include "SlateFwd.h"
 #include "Widgets/DeclarativeSyntaxSupport.h"
 #include "Widgets/SCompoundWidget.h"
@@ -17,13 +18,6 @@ class FConcertSessionActivitiesOptions;
  */
 class CONCERTSHAREDSLATE_API SConcertSessionRecovery : public SCompoundWidget
 {
-public:
-	/** Used to pull activities from a session. Used to fetch and display the activities of an archived session. */
-	using FFetchActivitiesFunc = TFunction<bool(TArray<TSharedPtr<FConcertSessionActivity>>& /*InOutActivities*/, int32& /*OutFetchedCount*/, FText& /*ErrorMsg*/)>;
-
-	/** Used to map an activity to its client. */
-	using FGetActivityClientInfoFunc = TFunction<const FConcertClientInfo*(FGuid /*ClientId*/)>;
-
 public:
 	SLATE_BEGIN_ARGS(SConcertSessionRecovery)
 		: _ClientAvatarColorColumnVisibility(EVisibility::Collapsed)
@@ -46,10 +40,10 @@ public:
 		SLATE_ARGUMENT(TSharedPtr<SWindow>, ParentWindow)
 
 		/** If bound, invoked iteratively to populate the activity list. */
-		SLATE_ARGUMENT(FFetchActivitiesFunc, OnFetchActivities)
+		SLATE_EVENT(SConcertSessionActivities::FFetchActivitiesFunc, OnFetchActivities)
 
 		/** If bound, invoked to map an activity to a client info. */
-		SLATE_ARGUMENT(FGetActivityClientInfoFunc, OnMapActivityToClient)
+		SLATE_EVENT(SConcertSessionActivities::FGetActivityClientInfoFunc, OnMapActivityToClient)
 
 		/** Invoked when the user clicks the 'recover' button. */
 		SLATE_ARGUMENT(TFunction<bool(TSharedPtr<FConcertSessionActivity>)>, OnRestore)
diff --git a/Engine/Plugins/Developer/Concert/ConcertUI/ConcertSharedSlate/Source/ConcertSharedSlate/Public/SessionHistory/SSessionHistory.h b/Engine/Plugins/Developer/Concert/ConcertUI/ConcertSharedSlate/Source/ConcertSharedSlate/Public/SessionHistory/SSessionHistory.h
--- a/Engine/Plugins/Developer/Concert/ConcertUI/ConcertSharedSlate/Source/ConcertSharedSlate/Public/SessionHistory/SSessionHistory.h
+++ b/Engine/Plugins/Developer/Concert/ConcertUI/ConcertSharedSlate/Source/ConcertSharedSlate/Public/SessionHistory/SSessionHistory.h
@@ -20,8 +20,8 @@ public:
 	
 	SLATE_BEGIN_ARGS(SSessionHistory) {}
 		SLATE_ARGUMENT(FName, PackageFilter)
-		SLATE_ARGUMENT(SConcertSessionActivities::FGetPackageEvent, GetPackageEvent)
-		SLATE_ARGUMENT(SConcertSessionActivities::FGetTransactionEvent, GetTransactionEvent)
+		SLATE_EVENT(SConcertSessionActivities::FGetPackageEvent, GetPackageEvent)
+		SLATE_EVENT(SConcertSessionActivities::FGetTransactionEvent, GetTransactionEvent)
 	SLATE_END_ARGS()
 
 	/**
@@ -50,12 +50,6 @@ private:
 	/** Returns the text to highlight when the search bar has a text set. */
 	FText HighlightSearchedText() const;
 
-	/** Returns the specified package event (without the package data itself) if available. */
-	bool GetPackageEvent(const FConcertSessionActivity& Activity, FConcertSyncPackageEventMetaData& OutPackageEvent) const;
-
-	/** Returns the specified package event if available. */
-	TFuture<TOptional<FConcertSyncTransactionEvent>> GetTransactionEvent(const FConcertSessionActivity& Activity) const;
-
 private:
 	
 	/** Holds the map of endpoint IDs to client info. */
@@ -78,4 +72,6 @@ private:
 
 	/** Used to limit activities to a given package only. */
 	FName PackageNameFilter;
+
+	TOptional<FConcertClientInfo> GetClientInfo(FGuid Guid) const;
 };
\ No newline at end of file
diff --git a/Engine/Plugins/Developer/Concert/DisasterRecoveryClient/Source/DisasterRecoveryClient/Private/SDisasterRecoveryHub.cpp b/Engine/Plugins/Developer/Concert/DisasterRecoveryClient/Source/DisasterRecoveryClient/Private/SDisasterRecoveryHub.cpp
--- a/Engine/Plugins/Developer/Concert/DisasterRecoveryClient/Source/DisasterRecoveryClient/Private/SDisasterRecoveryHub.cpp
+++ b/Engine/Plugins/Developer/Concert/DisasterRecoveryClient/Source/DisasterRecoveryClient/Private/SDisasterRecoveryHub.cpp
@@ -596,7 +596,7 @@ void SDisasterRecoveryHub::OnSessionSelectionChanged(TSharedPtr<FDisasterRecover
 TSharedRef<SWidget> SDisasterRecoveryHub::MakeSessionActivityView()
 {
 	return SAssignNew(ActivityView, SConcertSessionRecovery)
-		.OnFetchActivities([this](TArray<TSharedPtr<FConcertSessionActivity>>& InOutActivities, int32& OutFetchedCount, FText& ErrorMsg) { return FetchActivities(InOutActivities, OutFetchedCount, ErrorMsg); })
+		.OnFetchActivities(this, &SDisasterRecoveryHub::FetchActivities)
 		.ClientAvatarColorColumnVisibility(EVisibility::Hidden) // Disaster recovery has only one user, the local one.
 		.ClientNameColumnVisibility(EVisibility::Hidden)
 		.OperationColumnVisibility(EVisibility::Visible)
diff --git a/Engine/Plugins/Developer/Concert/MultiUserClient/Source/MultiUserClient/Private/Widgets/SConcertBrowser.cpp b/Engine/Plugins/Developer/Concert/MultiUserClient/Source/MultiUserClient/Private/Widgets/SConcertBrowser.cpp
--- a/Engine/Plugins/Developer/Concert/MultiUserClient/Source/MultiUserClient/Private/Widgets/SConcertBrowser.cpp
+++ b/Engine/Plugins/Developer/Concert/MultiUserClient/Source/MultiUserClient/Private/Widgets/SConcertBrowser.cpp
@@ -363,7 +363,8 @@ void FConcertClientSessionBrowserController::RestoreSession(const FGuid& ServerA
 	// The UI uses this function to map an activity ID from the stream to a client info.
 	auto GetActivityClientInfoFn = [ActivityStream](FGuid EndpointID)
 	{
-		return ActivityStream->GetActivityClientInfo(EndpointID);
+		const FConcertClientInfo* Result = ActivityStream->GetActivityClientInfo(EndpointID);
+		return Result ? TOptional<FConcertClientInfo>{ *Result } : TOptional<FConcertClientInfo>{};
 	};
 
 	// Invoked if the client selects a point in time to recover.
@@ -411,15 +412,14 @@ void FConcertClientSessionBrowserController::RestoreSession(const FGuid& ServerA
 			FAsyncTaskNotification Notification(NotificationConfig);
 			Notification.SetComplete(LOCTEXT("RecoveryError", "Failed to recover the session"), LOCTEXT("ClientUnavailable", "Concert client unavailable"), /*Success*/ false);
 		}
-
 		return bDismissRecoveryWindow;
 	};
 
 	TSharedRef<SConcertSessionRecovery> RestoreWidget = SNew(SConcertSessionRecovery)
 		.ParentWindow(NewWindow)
 		.IntroductionText(LOCTEXT("RecoverSessionIntroductionText", "Select the point in time at which the session should be restored"))
-		.OnFetchActivities(ReadActivitiesFn)
-		.OnMapActivityToClient(GetActivityClientInfoFn)
+		.OnFetchActivities_Lambda(ReadActivitiesFn)
+		.OnMapActivityToClient_Lambda(GetActivityClientInfoFn)
 		.OnRestore(OnAcceptRestoreFn)
 		.ClientNameColumnVisibility(EVisibility::Visible)
 		.ClientAvatarColorColumnVisibility(EVisibility::Visible)
