Improve directory structure

diff --git a/Engine/Plugins/Developer/Concert/ConcertUI/ConcertServerUI/Source/ConcertServerUI/Private/Widgets/Browser/ConcertServerSessionBrowserController.cpp b/Engine/Plugins/Developer/Concert/ConcertUI/ConcertServerUI/Source/ConcertServerUI/Private/Widgets/Browser/ConcertServerSessionBrowserController.cpp
--- a/Engine/Plugins/Developer/Concert/ConcertUI/ConcertServerUI/Source/ConcertServerUI/Private/Widgets/Browser/ConcertServerSessionBrowserController.cpp
+++ b/Engine/Plugins/Developer/Concert/ConcertUI/ConcertServerUI/Source/ConcertServerUI/Private/Widgets/Browser/ConcertServerSessionBrowserController.cpp
@@ -10,7 +10,7 @@
 #include "SConcertServerSessionBrowser.h"
 #include "Framework/Notifications/NotificationManager.h"
 #include "Misc/AsyncTaskNotification.h"
-#include "SessionBrowser/ConcertSessionItem.h"
+#include "Session/Browser/ConcertSessionItem.h"
 #include "Textures/SlateIcon.h"
 #include "Widgets/ConcertServerTabs.h"
 #include "Widgets/ConcertServerWindowController.h"
diff --git a/Engine/Plugins/Developer/Concert/ConcertUI/ConcertServerUI/Source/ConcertServerUI/Private/Widgets/Browser/ConcertServerSessionBrowserController.h b/Engine/Plugins/Developer/Concert/ConcertUI/ConcertServerUI/Source/ConcertServerUI/Private/Widgets/Browser/ConcertServerSessionBrowserController.h
--- a/Engine/Plugins/Developer/Concert/ConcertUI/ConcertServerUI/Source/ConcertServerUI/Private/Widgets/Browser/ConcertServerSessionBrowserController.h
+++ b/Engine/Plugins/Developer/Concert/ConcertUI/ConcertServerUI/Source/ConcertServerUI/Private/Widgets/Browser/ConcertServerSessionBrowserController.h
@@ -3,7 +3,7 @@
 #pragma once
 
 #include "CoreMinimal.h"
-#include "SessionBrowser/IConcertSessionBrowserController.h"
+#include "Session/Browser/IConcertSessionBrowserController.h"
 #include "Widgets/IConcertComponent.h"
 
 class IConcertServerSession;
diff --git a/Engine/Plugins/Developer/Concert/ConcertUI/ConcertServerUI/Source/ConcertServerUI/Private/Widgets/Browser/SConcertServerSessionBrowser.cpp b/Engine/Plugins/Developer/Concert/ConcertUI/ConcertServerUI/Source/ConcertServerUI/Private/Widgets/Browser/SConcertServerSessionBrowser.cpp
--- a/Engine/Plugins/Developer/Concert/ConcertUI/ConcertServerUI/Source/ConcertServerUI/Private/Widgets/Browser/SConcertServerSessionBrowser.cpp
+++ b/Engine/Plugins/Developer/Concert/ConcertUI/ConcertServerUI/Source/ConcertServerUI/Private/Widgets/Browser/SConcertServerSessionBrowser.cpp
@@ -2,13 +2,13 @@
 
 #include "SConcertServerSessionBrowser.h"
 
-#include "SessionBrowser/SConcertSessionBrowser.h"
+#include "Session/Browser/SConcertSessionBrowser.h"
 #include "Widgets/Browser/ConcertServerSessionBrowserController.h"
 #include "Widgets/StatusBar/SConcertStatusBar.h"
 
 #include "Dialog/SMessageDialog.h"
 #include "Framework/Docking/TabManager.h"
-#include "SessionBrowser/ConcertSessionItem.h"
+#include "Session/Browser/ConcertSessionItem.h"
 #include "Widgets/ConcertServerTabs.h"
 #include "Widgets/Input/SSearchBox.h"
 #include "Widgets/Layout/SBorder.h"
diff --git a/Engine/Plugins/Developer/Concert/ConcertUI/ConcertServerUI/Source/ConcertServerUI/Private/Widgets/Browser/SConcertServerSessionBrowser.h b/Engine/Plugins/Developer/Concert/ConcertUI/ConcertServerUI/Source/ConcertServerUI/Private/Widgets/Browser/SConcertServerSessionBrowser.h
--- a/Engine/Plugins/Developer/Concert/ConcertUI/ConcertServerUI/Source/ConcertServerUI/Private/Widgets/Browser/SConcertServerSessionBrowser.h
+++ b/Engine/Plugins/Developer/Concert/ConcertUI/ConcertServerUI/Source/ConcertServerUI/Private/Widgets/Browser/SConcertServerSessionBrowser.h
@@ -3,7 +3,7 @@
 #pragma once
 
 #include "CoreMinimal.h"
-#include "SessionBrowser/SConcertSessionBrowser.h"
+#include "Session/Browser/SConcertSessionBrowser.h"
 #include "Widgets/DeclarativeSyntaxSupport.h"
 #include "Widgets/SCompoundWidget.h"
 
diff --git a/Engine/Plugins/Developer/Concert/ConcertUI/ConcertServerUI/Source/ConcertServerUI/Private/Widgets/Session/PackageViewer/SConcertSessionPackageViewer.cpp b/Engine/Plugins/Developer/Concert/ConcertUI/ConcertServerUI/Source/ConcertServerUI/Private/Widgets/Session/PackageViewer/SConcertSessionPackageViewer.cpp
--- a/Engine/Plugins/Developer/Concert/ConcertUI/ConcertServerUI/Source/ConcertServerUI/Private/Widgets/Session/PackageViewer/SConcertSessionPackageViewer.cpp
+++ b/Engine/Plugins/Developer/Concert/ConcertUI/ConcertServerUI/Source/ConcertServerUI/Private/Widgets/Session/PackageViewer/SConcertSessionPackageViewer.cpp
@@ -2,7 +2,7 @@
 
 #include "SConcertSessionPackageViewer.h"
 
-#include "SConcertSessionActivities.h"
+#include "Session/Activity/SConcertSessionActivities.h"
 #include "Widgets/Input/SSearchBox.h"
 
 #define LOCTEXT_NAMESPACE "UnrealMultiUserUI"
diff --git a/Engine/Plugins/Developer/Concert/ConcertUI/ConcertServerUI/Source/ConcertServerUI/Private/Widgets/Session/PackageViewer/SConcertSessionPackageViewer.h b/Engine/Plugins/Developer/Concert/ConcertUI/ConcertServerUI/Source/ConcertServerUI/Private/Widgets/Session/PackageViewer/SConcertSessionPackageViewer.h
--- a/Engine/Plugins/Developer/Concert/ConcertUI/ConcertServerUI/Source/ConcertServerUI/Private/Widgets/Session/PackageViewer/SConcertSessionPackageViewer.h
+++ b/Engine/Plugins/Developer/Concert/ConcertUI/ConcertServerUI/Source/ConcertServerUI/Private/Widgets/Session/PackageViewer/SConcertSessionPackageViewer.h
@@ -3,7 +3,7 @@
 #pragma once
 
 #include "CoreMinimal.h"
-#include "SConcertSessionActivities.h"
+#include "Session/Activity/SConcertSessionActivities.h"
 #include "Widgets/SCompoundWidget.h"
 #include "Widgets/DeclarativeSyntaxSupport.h"
 
diff --git a/Engine/Plugins/Developer/Concert/ConcertUI/ConcertServerUI/Source/ConcertServerUI/Private/Widgets/Session/SConcertSessionInspector.cpp b/Engine/Plugins/Developer/Concert/ConcertUI/ConcertServerUI/Source/ConcertServerUI/Private/Widgets/Session/SConcertSessionInspector.cpp
--- a/Engine/Plugins/Developer/Concert/ConcertUI/ConcertServerUI/Source/ConcertServerUI/Private/Widgets/Session/SConcertSessionInspector.cpp
+++ b/Engine/Plugins/Developer/Concert/ConcertUI/ConcertServerUI/Source/ConcertServerUI/Private/Widgets/Session/SConcertSessionInspector.cpp
@@ -3,9 +3,8 @@
 #include "Widgets/Session/SConcertSessionInspector.h"
 
 #include "Framework/Docking/TabManager.h"
-#include "SessionHistory/SSessionHistory.h"
+#include "Session/History/SSessionHistory.h"
 #include "Widgets/Docking/SDockTab.h"
-#include "Widgets/SWindow.h"
 
 #define LOCTEXT_NAMESPACE "SConcertSessionInspector"
 
diff --git a/Engine/Plugins/Developer/Concert/ConcertUI/ConcertServerUI/Source/ConcertServerUI/Private/Widgets/Session/SConcertSessionInspector.h b/Engine/Plugins/Developer/Concert/ConcertUI/ConcertServerUI/Source/ConcertServerUI/Private/Widgets/Session/SConcertSessionInspector.h
--- a/Engine/Plugins/Developer/Concert/ConcertUI/ConcertServerUI/Source/ConcertServerUI/Private/Widgets/Session/SConcertSessionInspector.h
+++ b/Engine/Plugins/Developer/Concert/ConcertUI/ConcertServerUI/Source/ConcertServerUI/Private/Widgets/Session/SConcertSessionInspector.h
@@ -3,7 +3,7 @@
 #pragma once
 
 #include "CoreMinimal.h"
-#include "SessionHistory/SSessionHistory.h"
+#include "Session/History/SSessionHistory.h"
 #include "PackageViewer/SConcertSessionPackageViewer.h"
 #include "Widgets/SCompoundWidget.h"
 #include "Widgets/DeclarativeSyntaxSupport.h"
diff --git a/Engine/Plugins/Developer/Concert/ConcertUI/ConcertServerUI/Source/ConcertServerUI/Private/Widgets/Session/ServerSessionHistoryController.h b/Engine/Plugins/Developer/Concert/ConcertUI/ConcertServerUI/Source/ConcertServerUI/Private/Widgets/Session/ServerSessionHistoryController.h
--- a/Engine/Plugins/Developer/Concert/ConcertUI/ConcertServerUI/Source/ConcertServerUI/Private/Widgets/Session/ServerSessionHistoryController.h
+++ b/Engine/Plugins/Developer/Concert/ConcertUI/ConcertServerUI/Source/ConcertServerUI/Private/Widgets/Session/ServerSessionHistoryController.h
@@ -3,7 +3,7 @@
 #pragma once
 
 #include "CoreMinimal.h"
-#include "SessionHistory/AbstractSessionHistoryController.h"
+#include "Session/History/AbstractSessionHistoryController.h"
 
 struct FConcertSyncActivity;
 class FConcertSyncSessionDatabase;
diff --git a/Engine/Plugins/Developer/Concert/ConcertUI/ConcertSharedSlate/Source/ConcertSharedSlate/Private/SConcertSessionActivities.cpp b/Engine/Plugins/Developer/Concert/ConcertUI/ConcertSharedSlate/Source/ConcertSharedSlate/Private/SConcertSessionActivities.cpp
deleted file mode 100644
diff --git a/Engine/Plugins/Developer/Concert/ConcertUI/ConcertSharedSlate/Source/ConcertSharedSlate/Private/SConcertSessionRecovery.cpp b/Engine/Plugins/Developer/Concert/ConcertUI/ConcertSharedSlate/Source/ConcertSharedSlate/Private/SConcertSessionRecovery.cpp
--- a/Engine/Plugins/Developer/Concert/ConcertUI/ConcertSharedSlate/Source/ConcertSharedSlate/Private/SConcertSessionRecovery.cpp
+++ b/Engine/Plugins/Developer/Concert/ConcertUI/ConcertSharedSlate/Source/ConcertSharedSlate/Private/SConcertSessionRecovery.cpp
@@ -3,16 +3,13 @@
 #include "SConcertSessionRecovery.h"
 
 #include "EditorStyleSet.h"
-#include "EditorFontGlyphs.h"
+#include "Session/Activity/SConcertSessionActivities.h"
 #include "Widgets/SWindow.h"
-#include "Widgets/Input/SButton.h"
 #include "Widgets/Input/SSearchBox.h"
 #include "Widgets/Layout/SBorder.h"
 #include "Widgets/Layout/SBox.h"
-#include "Widgets/Layout/SSpacer.h"
 #include "Widgets/Layout/SSeparator.h"
 #include "Widgets/Layout/SUniformGridPanel.h"
-#include "SConcertSessionActivities.h"
 #include "SNegativeActionButton.h"
 #include "SPositiveActionButton.h"
 
diff --git a/Engine/Plugins/Developer/Concert/ConcertUI/ConcertSharedSlate/Source/ConcertSharedSlate/Private/Session/Activity/SConcertSessionActivities.cpp b/Engine/Plugins/Developer/Concert/ConcertUI/ConcertSharedSlate/Source/ConcertSharedSlate/Private/Session/Activity/SConcertSessionActivities.cpp
new file mode 100644
--- /dev/null
+++ b/Engine/Plugins/Developer/Concert/ConcertUI/ConcertSharedSlate/Source/ConcertSharedSlate/Private/Session/Activity/SConcertSessionActivities.cpp
@@ -0,0 +1,1178 @@
+// Copyright Epic Games, Inc. All Rights Reserved.
+
+#include "Session/Activity/SConcertSessionActivities.h"
+
+#include "Algo/Transform.h"
+#include "Misc/AsyncTaskNotification.h"
+#include "Misc/ITransaction.h"
+#include "Framework/Application/SlateApplication.h"
+#include "Framework/MultiBox/MultiBoxBuilder.h"
+#include "Widgets/SOverlay.h"
+#include "Widgets/SUndoHistoryDetails.h"
+#include "Widgets/Colors/SColorBlock.h"
+#include "Widgets/Input/SComboButton.h"
+#include "Widgets/Images/SThrobber.h"
+#include "Widgets/Layout/SBorder.h"
+#include "Widgets/Layout/SBox.h"
+#include "Widgets/Layout/SExpandableArea.h"
+#include "Widgets/Layout/SScrollBox.h"
+#include "Widgets/Layout/SSpacer.h"
+#include "Widgets/Layout/SSplitter.h"
+#include "Widgets/Text/SRichTextBlock.h"
+#include "Widgets/Text/STextBlock.h"
+#include "Widgets/Views/STableRow.h"
+#include "Styling/AppStyle.h"
+#include "ConcertFrontendStyle.h"
+#include "ConcertFrontendUtils.h"
+#include "ConcertLogGlobal.h"
+#include "ConcertTransactionEvents.h"
+#include "ConcertWorkspaceData.h"
+#include "SPackageDetails.h"
+
+#define LOCTEXT_NAMESPACE "SConcertSessionActivities"
+
+namespace ConcertSessionActivityUtils
+{
+// The columns names.
+const FName DateTimeColumnId    = TEXT("DateTime");
+const FName OperationColumnId   = TEXT("Operation");
+const FName PackageColumnId     = TEXT("Package");
+const FName SummaryColumnId     = TEXT("Summary");
+const FName ClientNameColumnId  = TEXT("Client");
+const FName AvatarColorColumnId = TEXT("Client AvatarColor");
+
+// The View Options check boxes.
+const FName DisplayRelativeTimeCheckBoxId       = TEXT("DisplayRelativeTime");
+const FName ShowConnectionActivitiesCheckBoxId  = TEXT("ShowConnectionActivities");
+const FName ShowLockActivitiesCheckBoxId        = TEXT("ShowLockActivities");
+const FName ShowPackageActivitiesCheckBoxId     = TEXT("ShowPackageActivities");
+const FName ShowTransactionActivitiesCheckBoxId = TEXT("ShowTransactionActivities");
+const FName ShowIgnoredActivitiesCheckBoxId     = TEXT("ShowIgnoredActivities");
+
+FText GetActivityDateTime(const FConcertSessionActivity& Activity, SConcertSessionActivities::ETimeFormat TimeFormat)
+{
+	return TimeFormat == SConcertSessionActivities::ETimeFormat::Relative ? ConcertFrontendUtils::FormatRelativeTime(Activity.Activity.EventTime) : FText::AsDateTime(Activity.Activity.EventTime);
+}
+
+FText GetOperationName(const FConcertSessionActivity& Activity)
+{
+	if (const FConcertSyncTransactionActivitySummary* TransactionSummary = Activity.ActivitySummary.Cast<FConcertSyncTransactionActivitySummary>())
+	{
+		return TransactionSummary->TransactionTitle;
+	}
+
+	if (const FConcertSyncPackageActivitySummary* PackageSummary = Activity.ActivitySummary.Cast<FConcertSyncPackageActivitySummary>())
+	{
+		auto GetSavedOperationNameFn = [](const FConcertSyncPackageActivitySummary* PackageSummary)
+		{
+			if (PackageSummary->bAutoSave)
+			{
+				return LOCTEXT("AutoSavePackageOperation", "Auto-Save Package");
+			}
+			return PackageSummary->bPreSave ? LOCTEXT("PreSavePackageOperation", "Pre-Save Package") : LOCTEXT("SavePackageOperation", "Save Package");
+		};
+
+		switch (PackageSummary->PackageUpdateType)
+		{
+		case EConcertPackageUpdateType::Added   : return LOCTEXT("NewPackageOperation",    "New Package");
+		case EConcertPackageUpdateType::Deleted : return LOCTEXT("DeletePackageOperation", "Delete Package");
+		case EConcertPackageUpdateType::Renamed : return LOCTEXT("RenamePackageOperation", "Rename Package");
+		case EConcertPackageUpdateType::Saved   : return GetSavedOperationNameFn(PackageSummary);
+		case EConcertPackageUpdateType::Dummy   : return LOCTEXT("DiscardPackageOperation", "Discard Changes");
+		default: break;
+		}
+	}
+
+	if (const FConcertSyncConnectionActivitySummary* ConnectionSummary = Activity.ActivitySummary.Cast<FConcertSyncConnectionActivitySummary>())
+	{
+		switch (ConnectionSummary->ConnectionEventType)
+		{
+		case EConcertSyncConnectionEventType::Connected:    return LOCTEXT("JoinOperation", "Join Session");
+		case EConcertSyncConnectionEventType::Disconnected: return LOCTEXT("LeaveOperation", "Leave Session");
+		default: break;
+		}
+	}
+
+	if (const FConcertSyncLockActivitySummary* LockSummary = Activity.ActivitySummary.Cast<FConcertSyncLockActivitySummary>())
+	{
+		switch (LockSummary->LockEventType)
+		{
+		case EConcertSyncLockEventType::Locked:   return LOCTEXT("LockOperation", "Lock");
+		case EConcertSyncLockEventType::Unlocked: return LOCTEXT("UnlockOperation", "Unlock");
+		default: break;
+		}
+	}
+
+	return FText::GetEmpty();
+}
+
+FText GetPackageName(const FConcertSessionActivity& Activity)
+{
+	if (const FConcertSyncPackageActivitySummary* Summary = Activity.ActivitySummary.Cast<FConcertSyncPackageActivitySummary>())
+	{
+		return FText::FromName(Summary->PackageName);
+	}
+
+	if (const FConcertSyncTransactionActivitySummary* Summary = Activity.ActivitySummary.Cast<FConcertSyncTransactionActivitySummary>())
+	{
+		return FText::FromName(Summary->PrimaryPackageName);
+	}
+
+	return FText::GetEmpty();
+}
+
+FText GetSummary(const FConcertSessionActivity& Activity, const FText& ClientName, bool bAsRichText)
+{
+	if (const FConcertSyncActivitySummary* Summary = Activity.ActivitySummary.Cast<FConcertSyncActivitySummary>())
+	{
+		return Summary->ToDisplayText(ClientName, bAsRichText);
+	}
+
+	return FText::GetEmpty();
+}
+
+FText GetClientName(const TOptional<FConcertClientInfo>& InActivityClient)
+{
+	return InActivityClient ? FText::AsCultureInvariant(InActivityClient->DisplayName) : FText::GetEmpty();
+}
+	
+FText GetClientName(const FConcertClientInfo* InActivityClient)
+{
+	return InActivityClient ? FText::AsCultureInvariant(InActivityClient->DisplayName) : FText::GetEmpty();
+}
+
+};
+
+/**
+ * Displays the summary of an activity recorded and recoverable in the SConcertSessionRecovery list view.
+ */
+class SConcertSessionActivityRow : public SMultiColumnTableRow<TSharedPtr<FConcertSessionActivity>>
+{
+public:
+	SLATE_BEGIN_ARGS(SConcertSessionActivityRow)
+		: _TimeFormat(SConcertSessionActivities::ETimeFormat::Relative)
+		, _HighlightText()
+		, _OnMakeColumnOverlayWidget()
+	{
+	}
+
+	SLATE_ATTRIBUTE(SConcertSessionActivities::ETimeFormat, TimeFormat)
+	SLATE_ATTRIBUTE(FText, HighlightText)
+	SLATE_ARGUMENT(SConcertSessionActivities::FMakeColumnOverlayWidgetFunc, OnMakeColumnOverlayWidget) // Function invoked when generating a row to add a widget above the column widget.
+
+	SLATE_END_ARGS()
+
+public:
+	/**
+	 * Constructs a row widget to display a Concert activity.
+	 * @param InArgs The widgets arguments.
+	 * @param InActivty The activity to display.
+	 * @param InActivityClient The client who produced this activity. Can be null if unknown or not desirable.
+	 * @param InOwnerTableView The table view that will own this row.
+	 */
+	void Construct(const FArguments& InArgs, TSharedPtr<FConcertSessionActivity> InActivity, const TOptional<FConcertClientInfo>& InActivityClient, const TSharedRef<STableViewBase>& InOwnerTableView);
+
+	/** Generates the widget representing this row. */
+	virtual TSharedRef<SWidget> GenerateWidgetForColumn(const FName& ColumnName) override;
+
+private:
+	/** Format the the time displayed as relative or absolute according to the attribute 'TimeFormat'. */
+	FText FormatEventDateTime() const;
+
+	/** Returns the client avatar color. */
+	FLinearColor GetClientAvatarColor() const { return ClientAvatarColor; }
+
+	/** Returns the display name of the client who performed the activity or an empty text if the information wasn't available or desired. */
+	FText GetClientName() const { return ClientName; }
+
+	/** Generates the tooltip for this row. */
+	FText MakeTooltipText() const;
+
+private:
+	TWeakPtr<FConcertSessionActivity> Activity;
+	TAttribute<SConcertSessionActivities::ETimeFormat> TimeFormat;
+	FText AbsoluteDateTime;
+	FText ClientName;
+	FLinearColor ClientAvatarColor;
+	TAttribute<FText> HighlightText;
+	SConcertSessionActivities::FMakeColumnOverlayWidgetFunc OnMakeColumnOverlayWidget;
+};
+
+
+void SConcertSessionActivityRow::Construct(const FArguments& InArgs, TSharedPtr<FConcertSessionActivity> InActivity, const TOptional<FConcertClientInfo>& InActivityClient, const TSharedRef<STableViewBase>& InOwnerTableView)
+{
+	Activity = InActivity;
+	TimeFormat = InArgs._TimeFormat;
+	HighlightText = InArgs._HighlightText;
+	OnMakeColumnOverlayWidget = InArgs._OnMakeColumnOverlayWidget;
+	AbsoluteDateTime = ConcertSessionActivityUtils::GetActivityDateTime(*InActivity, SConcertSessionActivities::ETimeFormat::Absolute); // Cache the absolute time. It doesn't changes.
+	ClientName = ConcertSessionActivityUtils::GetClientName(InActivityClient);
+	ClientAvatarColor = InActivityClient ? InActivityClient->AvatarColor : FConcertFrontendStyle::Get()->GetColor("Concert.DisconnectedColor");
+
+	// Construct base class
+	SMultiColumnTableRow<TSharedPtr<FConcertSessionActivity>>::Construct(FSuperRowType::FArguments(), InOwnerTableView);
+
+	if (InActivity->Activity.bIgnored)
+	{
+		SetColorAndOpacity(FLinearColor(0.5, 0.5, 0.5, 0.5));
+	}
+}
+
+TSharedRef<SWidget> SConcertSessionActivityRow::GenerateWidgetForColumn(const FName& ColumnId)
+{
+	TSharedPtr<FConcertSessionActivity> ActivityPin = Activity.Pin();
+	TSharedRef<SOverlay> Overlay = SNew(SOverlay);
+
+	if (ColumnId == ConcertSessionActivityUtils::AvatarColorColumnId)
+	{
+		Overlay->AddSlot()
+		.HAlign(HAlign_Center)
+		.VAlign(VAlign_Center)
+		.Padding(2, 1)
+		[
+			SNew(SColorBlock)
+			.Color(GetClientAvatarColor())
+			.Size(FVector2D(4.f, 16.f))
+		];
+	}
+	else if (ColumnId == ConcertSessionActivityUtils::DateTimeColumnId)
+	{
+		Overlay->AddSlot()
+		.VAlign(VAlign_Center)
+		[
+			SNew(STextBlock)
+			.Text(this, &SConcertSessionActivityRow::FormatEventDateTime)
+			.HighlightText(HighlightText)
+		];
+	}
+	else if (ColumnId == ConcertSessionActivityUtils::ClientNameColumnId)
+	{
+		Overlay->AddSlot()
+		.VAlign(VAlign_Center)
+		[
+			SNew(STextBlock)
+			.Text(GetClientName())
+			.HighlightText(HighlightText)
+		];
+	}
+	else if (ColumnId == ConcertSessionActivityUtils::PackageColumnId)
+	{
+		Overlay->AddSlot()
+		.VAlign(VAlign_Center)
+		[
+			SNew(STextBlock)
+			.Text(ConcertSessionActivityUtils::GetPackageName(*ActivityPin))
+			.HighlightText(HighlightText)
+		];
+	}
+	else if (ColumnId == ConcertSessionActivityUtils::OperationColumnId)
+	{
+		Overlay->AddSlot()
+		.VAlign(VAlign_Center)
+		[
+			SNew(STextBlock)
+			.Text(ConcertSessionActivityUtils::GetOperationName(*ActivityPin))
+			.HighlightText(HighlightText)
+		];
+	}
+	else
+	{
+		check(ColumnId == ConcertSessionActivityUtils::SummaryColumnId);
+
+		Overlay->AddSlot()
+		.VAlign(VAlign_Center)
+		[
+			SNew(SRichTextBlock)
+			.DecoratorStyleSet(FConcertFrontendStyle::Get().Get())
+			.Text(ConcertSessionActivityUtils::GetSummary(*ActivityPin, FText::GetEmpty(), true/*bAsRichText*/))
+			.HighlightText(HighlightText)
+		];
+	}
+
+	if (OnMakeColumnOverlayWidget.IsBound())
+	{
+		if (TSharedPtr<SWidget> OverlayedWidget = OnMakeColumnOverlayWidget.Execute(ActivityPin, ColumnId))
+		{
+			Overlay->AddSlot()
+			[
+				OverlayedWidget.ToSharedRef()
+			];
+		}
+	}
+
+	SetToolTipText(TAttribute<FText>(this, &SConcertSessionActivityRow::MakeTooltipText));
+
+	return Overlay;
+}
+
+FText SConcertSessionActivityRow::MakeTooltipText() const
+{
+	if (TSharedPtr<FConcertSessionActivity> ActivityPin = Activity.Pin())
+	{
+		FText Client = GetClientName();
+		FText Operation = ConcertSessionActivityUtils::GetOperationName(*ActivityPin);
+		FText Package = ConcertSessionActivityUtils::GetPackageName(*ActivityPin);
+		FText Summary = ConcertSessionActivityUtils::GetSummary(*ActivityPin, Client, /*bAsRichText*/false);
+
+		FTextBuilder TextBuilder;
+
+		if (!Operation.IsEmpty())
+		{
+			TextBuilder.AppendLine(Operation);
+		}
+
+		TextBuilder.AppendLineFormat(LOCTEXT("ActivityRowTooltip_DateTime", "{0} ({1})"), AbsoluteDateTime, ConcertFrontendUtils::FormatRelativeTime(ActivityPin->Activity.EventTime));
+
+		if (!Package.IsEmpty())
+		{
+			TextBuilder.AppendLine(Package);
+		}
+
+		if (!Summary.IsEmpty())
+		{
+			TextBuilder.AppendLine(Summary);
+		}
+
+		if (ActivityPin->Activity.bIgnored)
+		{
+			TextBuilder.AppendLine();
+			TextBuilder.AppendLine(LOCTEXT("IgnoredActivity", "** This activity cannot be recovered (likely recorded during a Multi-User session). It is displayed for crash inspection only. It will be ignored on restore."));
+		}
+
+		return TextBuilder.ToText();
+	}
+
+	return FText::GetEmpty();
+}
+
+FText SConcertSessionActivityRow::FormatEventDateTime() const
+{
+	if (TSharedPtr<FConcertSessionActivity> ItemPin = Activity.Pin())
+	{
+		return TimeFormat.Get() == SConcertSessionActivities::ETimeFormat::Relative ? ConcertFrontendUtils::FormatRelativeTime(ItemPin->Activity.EventTime) : AbsoluteDateTime;
+	}
+	return FText::GetEmpty();
+}
+
+
+void SConcertSessionActivities::Construct(const FArguments& InArgs)
+{
+	FetchActivitiesFn = InArgs._OnFetchActivities;
+	GetActivityUserFn = InArgs._OnMapActivityToClient;
+	GetTransactionEventFn = InArgs._OnGetTransactionEvent;
+	GetPackageEventFn = InArgs._OnGetPackageEvent;
+	MakeColumnOverlayWidgetFn = InArgs._OnMakeColumnOverlayWidget;
+	HighlightText = InArgs._HighlightText;
+	TimeFormat = InArgs._TimeFormat;
+	ClientNameColumnVisibility = InArgs._ClientNameColumnVisibility;
+	OperationColumnVisibility = InArgs._OperationColumnVisibility;
+	PackageColumnVisibility = InArgs._PackageColumnVisibility;
+	ConnectionActivitiesVisibility = InArgs._ConnectionActivitiesVisibility;
+	LockActivitiesVisibility = InArgs._LockActivitiesVisibility;
+	PackageActivitiesVisibility = InArgs._PackageActivitiesVisibility;
+	TransactionActivitiesVisibility = InArgs._TransactionActivitiesVisibility;
+	IgnoredActivitiesVisibility = InArgs._IgnoredActivitiesVisibility;
+	DetailsAreaVisibility = InArgs._DetailsAreaVisibility;
+	bAutoScrollDesired = InArgs._IsAutoScrollEnabled;
+
+	SearchTextFilter = MakeShared<TTextFilter<const FConcertSessionActivity&>>(TTextFilter<const FConcertSessionActivity&>::FItemToStringArray::CreateSP(this, &SConcertSessionActivities::PopulateSearchStrings));
+	SearchTextFilter->OnChanged().AddSP(this, &SConcertSessionActivities::OnActivityFilterUpdated);
+
+	// Set the initial filter state.
+	ActiveFilterFlags = QueryActiveActivityFilters();
+
+	// Create the table header. (Setting visibility on the column itself doesn't show/hide the column as one would expect, unfortunately)
+	TSharedRef<SHeaderRow> HeaderRow = SNew(SHeaderRow);
+	if (InArgs._ClientAvatarColorColumnVisibility.Get() == EVisibility::Visible)
+	{
+		HeaderRow->AddColumn(SHeaderRow::Column(ConcertSessionActivityUtils::AvatarColorColumnId)
+			.DefaultLabel(INVTEXT(""))
+			.ManualWidth(8));
+	}
+
+	HeaderRow->AddColumn(SHeaderRow::Column(ConcertSessionActivityUtils::DateTimeColumnId)
+		.DefaultLabel(LOCTEXT("DateTime", "Date/Time"))
+		.ManualWidth(160));
+
+	if (InArgs._ClientNameColumnVisibility.Get() == EVisibility::Visible)
+	{
+		HeaderRow->AddColumn(SHeaderRow::Column(ConcertSessionActivityUtils::ClientNameColumnId)
+			.DefaultLabel(LOCTEXT("Client", "Client"))
+			.ManualWidth(80));
+	}
+
+	if (InArgs._OperationColumnVisibility.Get() == EVisibility::Visible)
+	{
+		HeaderRow->AddColumn(SHeaderRow::Column(ConcertSessionActivityUtils::OperationColumnId)
+			.DefaultLabel(LOCTEXT("Operation", "Operation"))
+			.ManualWidth(160));
+	}
+
+	if (InArgs._PackageColumnVisibility.Get() == EVisibility::Visible)
+	{
+		HeaderRow->AddColumn(SHeaderRow::Column(ConcertSessionActivityUtils::PackageColumnId)
+			.DefaultLabel(LOCTEXT("Package", "Package"))
+			.ManualWidth(200));
+	}
+
+	HeaderRow->AddColumn(SHeaderRow::Column(ConcertSessionActivityUtils::SummaryColumnId)
+		.DefaultLabel(LOCTEXT("Summary", "Summary")));
+
+	ChildSlot
+	[
+		SNew(SSplitter)
+		.Orientation(Orient_Vertical)
+
+		// Activity List
+		+SSplitter::Slot()
+		.Value(0.75)
+		[
+			SNew(SOverlay)
+
+			+SOverlay::Slot() // Activity list itself.
+			[
+				SNew(SBorder)
+				.BorderImage(FAppStyle::Get().GetBrush("ToolPanel.GroupBorder"))
+				.BorderBackgroundColor(FSlateColor(FLinearColor(0.6, 0.6, 0.6)))
+				.Padding(0)
+				[
+					SAssignNew(ActivityView, SListView<TSharedPtr<FConcertSessionActivity>>)
+					.ListItemsSource(&Activities)
+					.OnGenerateRow(this, &SConcertSessionActivities::OnGenerateActivityRowWidget)
+					.SelectionMode(ESelectionMode::Single)
+					.AllowOverscroll(EAllowOverscroll::No)
+					.OnListViewScrolled(this, &SConcertSessionActivities::OnListViewScrolled)
+					.OnSelectionChanged(this, &SConcertSessionActivities::OnListViewSelectionChanged)
+					.HeaderRow(HeaderRow)
+				]
+			]
+
+			+SOverlay::Slot() // Display a reason why no activities are shown.
+			.HAlign(HAlign_Center)
+			.VAlign(VAlign_Center)
+			[
+				SNew(STextBlock)
+				.Visibility_Lambda([TextAttr = InArgs._NoActivitiesReasonText](){ return TextAttr.Get().IsEmptyOrWhitespace() ? EVisibility::Collapsed : EVisibility::Visible; })
+				.Text(InArgs._NoActivitiesReasonText)
+				.Justification(ETextJustify::Center)
+			]
+		]
+
+		// Activity details.
+		+SSplitter::Slot()
+		.Value(0.25)
+		.SizeRule(TAttribute<SSplitter::ESizeRule>(this, &SConcertSessionActivities::GetDetailsAreaSizeRule))
+		[
+			SAssignNew(ExpandableDetails, SExpandableArea)
+			.Visibility(GetDetailAreaVisibility())
+			.InitiallyCollapsed(true)
+			.BorderBackgroundColor(FLinearColor(0.6f, 0.6f, 0.6f, 1.0f))
+			.BorderImage_Lambda([this]() { return ConcertFrontendUtils::GetExpandableAreaBorderImage(*ExpandableDetails); })
+			.BodyBorderImage(FAppStyle::Get().GetBrush("ToolPanel.GroupBorder"))
+			.BodyBorderBackgroundColor(FLinearColor::White)
+			.OnAreaExpansionChanged(this, &SConcertSessionActivities::OnDetailsAreaExpansionChanged)
+			.Padding(0.0f)
+			.HeaderContent()
+			[
+				SNew(STextBlock)
+				.Text(LOCTEXT("Details", "Details"))
+				.Font(FAppStyle::Get().GetFontStyle("DetailsView.CategoryFontStyle"))
+				.ShadowOffset(FVector2D(1.0f, 1.0f))
+			]
+			.BodyContent()
+			[
+				SNew(SOverlay)
+
+				+SOverlay::Slot()
+				[
+					SNew(SScrollBox)
+					.ScrollBarThickness(FVector2D(12.0f, 5.0f)) // To have same thickness than the ListView scroll bar.
+
+					+SScrollBox::Slot()
+					[
+						SAssignNew(TransactionDetailsPanel, SUndoHistoryDetails)
+						.Visibility(EVisibility::Collapsed)
+					]
+
+					+SScrollBox::Slot()
+					[
+						SAssignNew(PackageDetailsPanel, SPackageDetails)
+						.Visibility(EVisibility::Collapsed)
+					]
+				]
+
+				+SOverlay::Slot()
+				[
+					SAssignNew(NoDetailsPanel, SBox)
+					.HAlign(HAlign_Center)
+					.VAlign(VAlign_Center)
+					.Visibility(EVisibility::Visible)
+					[
+						SNew(STextBlock)
+						.Text(this, &SConcertSessionActivities::GetNoDetailsText)
+					]
+				]
+
+				+SOverlay::Slot()
+				[
+					SAssignNew(LoadingDetailsPanel, SBox)
+					.HAlign(HAlign_Center)
+					.VAlign(VAlign_Center)
+					.Visibility(EVisibility::Collapsed)
+					[
+						SNew(SThrobber)
+					]
+				]
+			]
+		]
+	];
+
+	// Check if some activities are already available.
+	FetchActivities();
+
+	if (bAutoScrollDesired)
+	{
+		FSlateApplication::Get().OnPostTick().AddSP(this, &SConcertSessionActivities::OnPostTick);
+	}
+}
+
+void SConcertSessionActivities::Tick(const FGeometry& AllottedGeometry, const double InCurrentTime, const float InDeltaTime)
+{
+	EConcertActivityFilterFlags LatestFilterFlags = QueryActiveActivityFilters();
+	if (ActiveFilterFlags != LatestFilterFlags)
+	{
+		ActiveFilterFlags = LatestFilterFlags;
+		OnActivityFilterUpdated();
+	}
+
+	FetchActivities(); // Check if we should fetch more activities in case we filtered out to many of them.
+}
+
+FText SConcertSessionActivities::GetNoDetailsText() const
+{
+	return GetSelectedActivity() ?
+		LOCTEXT("NoDetails_NotAvailable", "The selected activity doesn't have details to display.") :
+		LOCTEXT("NoDetails_NoActivitySelected", "Select an activity to view its details.");
+}
+
+void SConcertSessionActivities::OnPostTick(float)
+{
+	// NOTE: The way the list view adjust the scroll position when the component is resized has some behaviors to consider to get auto-scrolling
+	//       working consistently. When the list view shrink (allowing less item) the scroll view doesn't remain anchored at the end. Instead
+	//       the scroll position moves a little bit up and the list view doesn't consider it as scrolling because OnListViewScrolled() is not
+	//       called. The code below detects that case and maintain the scroll position at the end when required.
+
+	if (bActivityViewScrolled) // OnListViewScrolled() was invoked. The user scrolled the activity list or enlarged the view to see more items.
+	{
+		bUserScrolling = !ActivityView->GetScrollDistanceRemaining().IsNearlyZero();
+		bActivityViewScrolled = false;
+	}
+	else if (bAutoScrollDesired && !bUserScrolling && ActivityView->GetScrollDistanceRemaining().Y > 0) // See NOTE.
+	{
+		ActivityView->ScrollToBottom(); // Ensure the scroll position is maintained at the end.
+	}
+}
+
+void SConcertSessionActivities::OnListViewScrolled(double InScrollOffset)
+{
+	bActivityViewScrolled = true;
+
+	if (FetchActivitiesFn.IsBound()) // This widget is responsible to populate the view.
+	{
+		if (!bAllActivitiesFetched && ActivityView->GetScrollDistance().Y > 0.7) // Should fetch more?
+		{
+			DesiredActivitiesCount += ActivitiesPerRequest; // This will request another 'page' the next time FetchActivities() is called.
+		}
+	}
+}
+
+void SConcertSessionActivities::OnListViewSelectionChanged(TSharedPtr<FConcertSessionActivity> InActivity, ESelectInfo::Type SelectInfo)
+{
+	UpdateDetailArea(InActivity);
+}
+
+void SConcertSessionActivities::OnDetailsAreaExpansionChanged(bool bExpanded)
+{
+	bDetailsAreaExpanded = bExpanded;
+	UpdateDetailArea(bDetailsAreaExpanded ? GetSelectedActivity() : nullptr);
+}
+
+void SConcertSessionActivities::UpdateDetailArea(TSharedPtr<FConcertSessionActivity> InSelectedActivity)
+{
+	if (DetailsAreaVisibility != EVisibility::Visible || !bDetailsAreaExpanded)
+	{
+		return;
+	}
+	else if (!InSelectedActivity.IsValid()) // The selection was cleared?
+	{
+		SetDetailsPanelVisibility(NoDetailsPanel.Get());
+	}
+	else if (InSelectedActivity->EventPayload) // The event payload is already bundled in the activity stream?
+	{
+		if (InSelectedActivity->Activity.EventType == EConcertSyncActivityEventType::Transaction)
+		{
+			FConcertSyncTransactionEvent TransactionEvent;
+			InSelectedActivity->EventPayload->GetTypedPayload(TransactionEvent);
+			if (TransactionEvent.Transaction.ExportedObjects.Num())
+			{
+				DisplayTransactionDetails(*InSelectedActivity, TransactionEvent.Transaction);
+			}
+			else
+			{
+				SetDetailsPanelVisibility(NoDetailsPanel.Get());
+			}
+		}
+		else if (InSelectedActivity->Activity.EventType == EConcertSyncActivityEventType::Package)
+		{
+			FConcertSyncPackageEvent PackageEvent;
+			InSelectedActivity->EventPayload->GetTypedPayload(PackageEvent);
+			checkf(!PackageEvent.Package.HasPackageData(), TEXT("UI should only request the package meta data because the package data is not useful and can be very large"));
+			DisplayPackageDetails(*InSelectedActivity, PackageEvent.PackageRevision, PackageEvent.Package.Info);
+		}
+		else // Other activity types (lock/connection) don't have details panel.
+		{
+			SetDetailsPanelVisibility(NoDetailsPanel.Get());
+		}
+	}
+	else if (InSelectedActivity->Activity.EventType == EConcertSyncActivityEventType::Transaction && GetTransactionEventFn.IsBound()) // A function is bound to get the transaction event?
+	{
+		SetDetailsPanelVisibility(LoadingDetailsPanel.Get());
+		TWeakPtr<SConcertSessionActivities> WeakSelf = SharedThis(this);
+		GetTransactionEventFn.Execute(*InSelectedActivity).Next([WeakSelf, InSelectedActivity](const TOptional<FConcertSyncTransactionEvent>& TransactionEvent)
+		{
+			if (TSharedPtr<SConcertSessionActivities> Self = WeakSelf.Pin()) // If 'this' object hasn't been deleted.
+			{
+				if (Self->GetSelectedActivity() == InSelectedActivity) // Ensure the activity is still selected.
+				{
+					if (TransactionEvent.IsSet() && TransactionEvent->Transaction.ExportedObjects.Num())
+					{
+						Self->DisplayTransactionDetails(*InSelectedActivity, TransactionEvent.GetValue().Transaction);
+					}
+					else
+					{
+						Self->SetDetailsPanelVisibility(Self->NoDetailsPanel.Get());
+					}
+				}
+				// else -> The details panel is presenting information for another activity (or no activity).
+			}
+			// else -> The widget was deleted.
+		});
+	}
+	else if (InSelectedActivity->Activity.EventType == EConcertSyncActivityEventType::Package && GetPackageEventFn.IsBound()) // A function is bound to get the package event?
+	{
+		SetDetailsPanelVisibility(LoadingDetailsPanel.Get());
+		FConcertSyncPackageEventMetaData PackageEventMetaData;
+		if (GetPackageEventFn.Execute(*InSelectedActivity, PackageEventMetaData))
+		{
+			DisplayPackageDetails(*InSelectedActivity, PackageEventMetaData.PackageRevision, PackageEventMetaData.PackageInfo);
+		}
+		else
+		{
+			SetDetailsPanelVisibility(NoDetailsPanel.Get());
+		}
+	}
+	else
+	{
+		SetDetailsPanelVisibility(NoDetailsPanel.Get());
+	}
+}
+
+EConcertActivityFilterFlags SConcertSessionActivities::QueryActiveActivityFilters() const
+{
+	// The visibility attributes are externally provided. (In practice, they are controlled from the 'View Options' check boxes).
+	EConcertActivityFilterFlags ActiveFlags = EConcertActivityFilterFlags::ShowAll;
+
+	if (ConnectionActivitiesVisibility.Get() != EVisibility::Visible)
+	{
+		ActiveFlags |= EConcertActivityFilterFlags::HideConnectionActivities;
+	}
+	if (LockActivitiesVisibility.Get() != EVisibility::Visible)
+	{
+		ActiveFlags |= EConcertActivityFilterFlags::HideLockActivities;
+	}
+	if (PackageActivitiesVisibility.Get() != EVisibility::Visible)
+	{
+		ActiveFlags |= EConcertActivityFilterFlags::HidePackageActivities;
+	}
+	if (TransactionActivitiesVisibility.Get() != EVisibility::Visible)
+	{
+		ActiveFlags |= EConcertActivityFilterFlags::HideTransactionActivities;
+	}
+	if (IgnoredActivitiesVisibility.Get() != EVisibility::Visible)
+	{
+		ActiveFlags |= EConcertActivityFilterFlags::HideIgnoredActivities;
+	}
+
+	return ActiveFlags;
+}
+
+void SConcertSessionActivities::OnActivityFilterUpdated()
+{
+	// Try preserving the selected activity.
+	TSharedPtr<FConcertSessionActivity> SelectedActivity = GetSelectedActivity();
+
+	// Reset the list of displayed activities.
+	Activities.Reset(AllActivities.Num());
+
+	// Apply the filter.
+	for (TSharedPtr<FConcertSessionActivity>& Activity : AllActivities)
+	{
+		if (PassesFilters(*Activity))
+		{
+			Activities.Add(Activity);
+		}
+	}
+
+	// Restore/reset the selected activity.
+	if (SelectedActivity && Activities.Contains(SelectedActivity))
+	{
+		ActivityView->SetItemSelection(SelectedActivity, true); // Restore previous selection.
+		ActivityView->RequestScrollIntoView(SelectedActivity);
+	}
+	else if (bAutoScrollDesired && !bUserScrolling) // No activity was selected.
+	{
+		ActivityView->ScrollToBottom();
+	}
+
+	ActivityView->RequestListRefresh();
+}
+
+void SConcertSessionActivities::FetchActivities()
+{
+	if (!FetchActivitiesFn.IsBound()) // Not bound?
+	{
+		return; // The widget is expected to be populated/cleared externally using Append()/Reset()
+	}
+
+	bool bRefresh = false;
+
+	// If they are still activities to fetch and the user scrolled down (or our nominal amount is not reached), request more from the server.
+	if (!bAllActivitiesFetched && (Activities.Num() < DesiredActivitiesCount))
+	{
+		FText ErrorMsg;
+		int32 FetchCount = 0; // The number of activities fetched in this iteration.
+		int32 StartInsertPos = AllActivities.Num();
+
+		bAllActivitiesFetched = FetchActivitiesFn.Execute(AllActivities, FetchCount, ErrorMsg);
+		if (ErrorMsg.IsEmpty())
+		{
+			if (FetchCount) // New activities appended?
+			{
+				for (int32 Index = StartInsertPos; Index < AllActivities.Num(); ++Index) // Append the fetched activities
+				{
+					if (PassesFilters(*AllActivities[Index]))
+					{
+						Activities.Add(AllActivities[Index]);
+						bRefresh = true;
+					}
+
+					if (AllActivities[Index]->Activity.bIgnored)
+					{
+						++IgnoredActivityNum;
+					}
+				}
+			}
+		}
+		else
+		{
+			FAsyncTaskNotificationConfig NotificationConfig;
+			NotificationConfig.bIsHeadless = false;
+			NotificationConfig.bKeepOpenOnFailure = true;
+			NotificationConfig.LogCategory = &LogConcert;
+
+			FAsyncTaskNotification Notification(NotificationConfig);
+			Notification.SetComplete(LOCTEXT("FetchError", "Failed to retrieve session activities"), ErrorMsg, /*Success*/ false);
+		}
+	}
+
+	if (Activities.Num() && ActivityView->GetSelectedItems().Num() == 0)
+	{
+		ActivityView->SetItemSelection(Activities[0], true);
+	}
+
+	if (bRefresh)
+	{
+		if (bAutoScrollDesired && !bUserScrolling)
+		{
+			ActivityView->ScrollToBottom();
+		}
+
+		ActivityView->RequestListRefresh();
+	}
+}
+
+void SConcertSessionActivities::Append(TSharedPtr<FConcertSessionActivity> Activity)
+{
+	if (Activity->Activity.bIgnored)
+	{
+		++IgnoredActivityNum;
+	}
+
+	AllActivities.Add(Activity);
+	if (PassesFilters(*Activity))
+	{
+		Activities.Add(MoveTemp(Activity));
+
+		if (bAutoScrollDesired && !bUserScrolling)
+		{
+			ActivityView->ScrollToBottom();
+		}
+
+		ActivityView->RequestListRefresh();
+	}
+}
+
+void SConcertSessionActivities::RequestRefresh()
+{
+	ActivityView->RequestListRefresh();
+}
+
+void SConcertSessionActivities::ResetActivityList()
+{
+	Activities.Reset();
+	AllActivities.Reset();
+	ActivityView->RequestListRefresh();
+	bAllActivitiesFetched = false;
+	bUserScrolling = false;
+	DesiredActivitiesCount = ActivitiesPerRequest;
+	IgnoredActivityNum = 0;
+}
+
+bool SConcertSessionActivities::PassesFilters(const FConcertSessionActivity& Activity)
+{
+	if (Activity.Activity.EventType == EConcertSyncActivityEventType::Connection && ConnectionActivitiesVisibility.Get() != EVisibility::Visible) // Filter out 'connection' activities?
+	{
+		return false;
+	}
+	else if (Activity.Activity.EventType == EConcertSyncActivityEventType::Lock && LockActivitiesVisibility.Get() != EVisibility::Visible) // Filter out 'lock' activities?
+	{
+		return false;
+	}
+	else if (Activity.Activity.EventType == EConcertSyncActivityEventType::Package && PackageActivitiesVisibility.Get() != EVisibility::Visible) // Filter out 'package' activities?
+	{
+		return false;
+	}
+	else if (Activity.Activity.EventType == EConcertSyncActivityEventType::Transaction && TransactionActivitiesVisibility.Get() != EVisibility::Visible) // Filter out 'transaction' activities?
+	{
+		return false;
+	}
+	else if (Activity.Activity.bIgnored && IgnoredActivitiesVisibility.Get() != EVisibility::Visible) // Filter out 'ignored' activities?
+	{
+		return false;
+	}
+
+	return SearchTextFilter->PassesFilter(Activity);
+}
+
+FText SConcertSessionActivities::UpdateTextFilter(const FText& InFilterText)
+{
+	SearchTextFilter->SetRawFilterText(InFilterText);
+	return SearchTextFilter->GetFilterErrorText();
+}
+
+void SConcertSessionActivities::PopulateSearchStrings(const FConcertSessionActivity& Activity, TArray<FString>& OutSearchStrings) const
+{
+	FText ClientName = GetActivityUserFn.IsBound()
+		? ConcertSessionActivityUtils::GetClientName(GetActivityUserFn.Execute(Activity.Activity.EndpointId)) : FText::GetEmpty();
+
+	OutSearchStrings.Add(ConcertSessionActivityUtils::GetActivityDateTime(Activity, TimeFormat.Get()).ToString());
+	OutSearchStrings.Add(ConcertSessionActivityUtils::GetSummary(Activity, ClientName, /*bAsRichText*/false).ToString());
+
+	if (ClientNameColumnVisibility.Get() == EVisibility::Visible)
+	{
+		OutSearchStrings.Add(ClientName.ToString());
+	}
+
+	if (OperationColumnVisibility.Get() == EVisibility::Visible)
+	{
+		OutSearchStrings.Add(ConcertSessionActivityUtils::GetOperationName(Activity).ToString());
+	}
+
+	if (PackageColumnVisibility.Get() == EVisibility::Visible)
+	{
+		OutSearchStrings.Add(ConcertSessionActivityUtils::GetPackageName(Activity).ToString());
+	}
+}
+
+TSharedRef<ITableRow> SConcertSessionActivities::OnGenerateActivityRowWidget(TSharedPtr<FConcertSessionActivity> Activity, const TSharedRef<STableViewBase>& OwnerTable)
+{
+	const TOptional<FConcertClientInfo> ActivityClient = GetActivityUserFn.IsBound() ? GetActivityUserFn.Execute(Activity->Activity.EndpointId) : TOptional<FConcertClientInfo>{};
+	return SNew(SConcertSessionActivityRow, Activity, ActivityClient, OwnerTable)
+		.TimeFormat(TimeFormat)
+		.HighlightText(HighlightText)
+		.OnMakeColumnOverlayWidget(MakeColumnOverlayWidgetFn);
+}
+
+TSharedPtr<FConcertSessionActivity> SConcertSessionActivities::GetSelectedActivity() const
+{
+	TArray<TSharedPtr<FConcertSessionActivity>> SelectedItems = ActivityView->GetSelectedItems();
+	return SelectedItems.Num() ? SelectedItems[0] : nullptr;
+}
+
+TSharedPtr<FConcertSessionActivity> SConcertSessionActivities::GetMostRecentActivity() const
+{
+	// NOTE: This function assumes that activities are sorted by ID. When used for recovery purpose,
+	//       the activities are listed from the most recent to the oldest. When displaying a live
+	//       session activity stream, the activities are listed from the oldest to the newest.
+
+	if (AllActivities.Num()) // Ignore the filters.
+	{
+		if (AllActivities[0]->Activity.ActivityId > AllActivities.Last()->Activity.ActivityId)
+		{
+			return AllActivities[0]; // Listed from the latest to oldest.
+		}
+		return AllActivities.Last(); // Listed from the oldest to latest.
+	}
+	return nullptr; // The list is empty.
+}
+
+bool SConcertSessionActivities::IsLastColumn(const FName& ColumnId) const
+{
+	return ColumnId == ConcertSessionActivityUtils::SummaryColumnId; // Summary column is always visible and always the last.
+}
+
+void SConcertSessionActivities::DisplayTransactionDetails(const FConcertSessionActivity& Activity, const FConcertTransactionEventBase& InTransaction)
+{
+	const FConcertSyncTransactionActivitySummary* Summary = Activity.ActivitySummary.Cast<FConcertSyncTransactionActivitySummary>();
+	FString TransactionTitle = Summary ? Summary->TransactionTitle.ToString() : FString();
+
+	FTransactionDiff TransactionDiff{ InTransaction.TransactionId, TransactionTitle };
+
+	for (const FConcertExportedObject& ExportedObject : InTransaction.ExportedObjects)
+	{
+		FTransactionObjectDeltaChange DeltaChange;
+		Algo::Transform(ExportedObject.PropertyDatas, DeltaChange.ChangedProperties, [](const FConcertSerializedPropertyData& PropertyData) { return PropertyData.PropertyName; });
+
+		DeltaChange.bHasNameChange = ExportedObject.ObjectData.NewName != FName();
+		DeltaChange.bHasOuterChange = ExportedObject.ObjectData.NewOuterPathName != FName();
+		DeltaChange.bHasExternalPackageChange = ExportedObject.ObjectData.NewExternalPackageName != FName();
+		DeltaChange.bHasPendingKillChange = ExportedObject.ObjectData.bIsPendingKill;
+
+		FString ObjectPathName = ExportedObject.ObjectId.ObjectOuterPathName.ToString() + TEXT(".") + ExportedObject.ObjectId.ObjectName.ToString();
+		TSharedPtr<FTransactionObjectEvent> Event = MakeShared<FTransactionObjectEvent>(InTransaction.TransactionId, InTransaction.OperationId, ETransactionObjectEventType::Finalized, MoveTemp(DeltaChange), nullptr
+			, ExportedObject.ObjectId.ObjectPackageName
+			, ExportedObject.ObjectId.ObjectName
+			, FName(*MoveTemp(ObjectPathName))
+			, ExportedObject.ObjectId.ObjectOuterPathName
+			, ExportedObject.ObjectId.ObjectExternalPackageName
+			, ExportedObject.ObjectId.ObjectClassPathName);
+
+		TransactionDiff.DiffMap.Emplace(FName(*ObjectPathName), MoveTemp(Event));
+	}
+
+	TransactionDetailsPanel->SetSelectedTransaction(MoveTemp(TransactionDiff));
+	SetDetailsPanelVisibility(TransactionDetailsPanel.Get());
+}
+
+void SConcertSessionActivities::DisplayPackageDetails(const FConcertSessionActivity& Activity, int64 PackageRevision, const FConcertPackageInfo& PackageInfo)
+{
+	const TOptional<FConcertClientInfo> ClientInfo = GetActivityUserFn.IsBound()
+		? GetActivityUserFn.Execute(Activity.Activity.EndpointId) : TOptional<FConcertClientInfo>{};
+	PackageDetailsPanel->SetPackageInfo(PackageInfo, PackageRevision, ClientInfo ? ClientInfo->DisplayName : FString());
+	SetDetailsPanelVisibility(PackageDetailsPanel.Get());
+}
+
+void SConcertSessionActivities::SetDetailsPanelVisibility(const SWidget* VisiblePanel)
+{
+	TransactionDetailsPanel->SetVisibility(VisiblePanel == TransactionDetailsPanel.Get() ? EVisibility::Visible : EVisibility::Collapsed);
+	PackageDetailsPanel->SetVisibility(VisiblePanel == PackageDetailsPanel.Get() ? EVisibility::Visible : EVisibility::Collapsed);
+	NoDetailsPanel->SetVisibility(VisiblePanel == NoDetailsPanel.Get() ? EVisibility::Visible : EVisibility::Collapsed);
+	LoadingDetailsPanel->SetVisibility(VisiblePanel == LoadingDetailsPanel.Get() ? EVisibility::Visible : EVisibility::Collapsed);
+}
+
+TSharedRef<SWidget> FConcertSessionActivitiesOptions::MakeMenuWidget()
+{
+	FMenuBuilder MenuBuilder(/*bInShouldCloseWindowAfterMenuSelection=*/true, nullptr);
+
+	MenuBuilder.AddMenuEntry(
+		LOCTEXT("DisplayRelativeTime", "Display Relative Time"),
+		LOCTEXT("DisplayRelativeTime_Tooltip", "Displays Time Relative to the Current Time"),
+		FSlateIcon(),
+		FUIAction(
+			FExecuteAction::CreateSP(this, &FConcertSessionActivitiesOptions::OnOptionToggled, ConcertSessionActivityUtils::DisplayRelativeTimeCheckBoxId),
+			FCanExecuteAction::CreateLambda([] { return true; }),
+			FIsActionChecked::CreateLambda([this] { return bDisplayRelativeTime; })),
+		NAME_None,
+		EUserInterfaceActionType::ToggleButton
+	);
+
+	if (bEnablePackageActivityFiltering)
+	{
+		MenuBuilder.AddMenuEntry(
+			LOCTEXT("ShowPackageActivities", "Show Package Activities"),
+			LOCTEXT("ShowPackageActivities_Tooltip", "Displays create/save/rename/delete package events."),
+			FSlateIcon(),
+			FUIAction(
+				FExecuteAction::CreateSP(this, &FConcertSessionActivitiesOptions::OnOptionToggled, ConcertSessionActivityUtils::ShowPackageActivitiesCheckBoxId),
+				FCanExecuteAction::CreateLambda([] { return true; }),
+				FIsActionChecked::CreateLambda([this] { return bDisplayPackageActivities; })),
+			NAME_None,
+			EUserInterfaceActionType::ToggleButton
+		);
+	}
+
+	if (bEnableTransactionActivityFiltering)
+	{
+		MenuBuilder.AddMenuEntry(
+			LOCTEXT("ShowTransactionActivities", "Show Transaction Activities"),
+			LOCTEXT("ShowTransactionActivities_Tooltip", "Displays changes performed on assets."),
+			FSlateIcon(),
+			FUIAction(
+				FExecuteAction::CreateSP(this, &FConcertSessionActivitiesOptions::OnOptionToggled, ConcertSessionActivityUtils::ShowTransactionActivitiesCheckBoxId),
+				FCanExecuteAction::CreateLambda([] { return true; }),
+				FIsActionChecked::CreateLambda([this] { return bDisplayTransactionActivities; })),
+			NAME_None,
+			EUserInterfaceActionType::ToggleButton
+		);
+	}
+
+	if (bEnableConnectionActivityFiltering)
+	{
+		MenuBuilder.AddMenuEntry(
+			LOCTEXT("ShowConnectionActivities", "Show Connection Activities"),
+			LOCTEXT("ShowConnectionActivities_Tooltip", "Displays when client joined or left the session"),
+			FSlateIcon(),
+			FUIAction(
+				FExecuteAction::CreateSP(this, &FConcertSessionActivitiesOptions::OnOptionToggled, ConcertSessionActivityUtils::ShowConnectionActivitiesCheckBoxId),
+				FCanExecuteAction::CreateLambda([] { return true; }),
+				FIsActionChecked::CreateLambda([this] { return bDisplayConnectionActivities; })),
+			NAME_None,
+			EUserInterfaceActionType::ToggleButton
+		);
+	}
+
+	if (bEnableLockActivityFiltering)
+	{
+		MenuBuilder.AddMenuEntry(
+			LOCTEXT("ShowLockActivities", "Show Lock Activities"),
+			LOCTEXT("ShowLockActivities_Tooltip", "Displays lock/unlock events"),
+			FSlateIcon(),
+			FUIAction(
+				FExecuteAction::CreateSP(this, &FConcertSessionActivitiesOptions::OnOptionToggled, ConcertSessionActivityUtils::ShowLockActivitiesCheckBoxId),
+				FCanExecuteAction::CreateLambda([] { return true; }),
+				FIsActionChecked::CreateLambda([this] { return bDisplayLockActivities; })),
+			NAME_None,
+			EUserInterfaceActionType::ToggleButton
+		);
+	}
+
+	if (bEnableIgnoredActivityFiltering)
+	{
+		MenuBuilder.AddMenuEntry(
+			LOCTEXT("ShowIgnoredActivities", "Show Unrecoverable Activities"),
+			LOCTEXT("ShowIgnoredActivities_Tooltip", "Displays activities that were recorded, but could not be recovered in this context."),
+			FSlateIcon(),
+			FUIAction(
+				FExecuteAction::CreateSP(this, &FConcertSessionActivitiesOptions::OnOptionToggled, ConcertSessionActivityUtils::ShowIgnoredActivitiesCheckBoxId),
+				FCanExecuteAction::CreateLambda([] { return true; }),
+				FIsActionChecked::CreateLambda([this] { return bDisplayIgnoredActivities; })),
+			NAME_None,
+			EUserInterfaceActionType::ToggleButton
+		);
+	}
+	
+	return MenuBuilder.MakeWidget();
+}
+
+TSharedRef<SWidget> FConcertSessionActivitiesOptions::MakeViewOptionsWidget()
+{
+	return SNew(SComboButton)
+	.ComboButtonStyle(FAppStyle::Get(), "GenericFilters.ComboButtonStyle")
+	.ForegroundColor(FLinearColor::White)
+	.ContentPadding(0)
+	.OnGetMenuContent(this, &FConcertSessionActivitiesOptions::MakeMenuWidget)
+	.HasDownArrow(true)
+	.ContentPadding(FMargin(1, 0))
+	.ButtonContent()
+	[
+		SNew(SHorizontalBox)
+
+		+SHorizontalBox::Slot()
+		.AutoWidth()
+		.VAlign(VAlign_Center)
+		[
+			SNew(SImage).Image(FAppStyle::Get().GetBrush("GenericViewButton")) // The eye ball image.
+		]
+
+		+SHorizontalBox::Slot()
+		.AutoWidth()
+		.Padding(2, 0, 0, 0)
+		.VAlign(VAlign_Center)
+		[
+			SNew(STextBlock).Text(LOCTEXT("ViewOptions", "View Options"))
+		]
+	];
+}
+
+TSharedRef<SWidget> FConcertSessionActivitiesOptions::MakeDisplayedActivityCountWidget(TAttribute<int32> Total, TAttribute<int32> Displayed)
+{
+	return SNew(STextBlock)
+		.Text_Lambda([Total = MoveTemp(Total), Displayed = MoveTemp(Displayed)]()
+		{
+			if (Total.Get() == Displayed.Get())
+			{
+				return FText::Format(LOCTEXT("OperationCount", "{0} operations"), Total.Get());
+			}
+			else
+			{
+				return FText::Format(LOCTEXT("PartialOperationCount", "Showing {0} of {1} {1}|plural(one=operation,other=operations)"), Displayed.Get(), Total.Get());
+			}
+		});
+}
+
+TSharedRef<SWidget> FConcertSessionActivitiesOptions::MakeStatusBar(TAttribute<int32> Total, TAttribute<int32> Displayed)
+{
+	return SNew(SHorizontalBox)
+
+	// Operation count.
+	+SHorizontalBox::Slot()
+	.AutoWidth()
+	.VAlign(VAlign_Center)
+	[
+		MakeDisplayedActivityCountWidget(MoveTemp(Total), MoveTemp(Displayed))
+	]
+
+	// Gap Filler
+	+SHorizontalBox::Slot()
+	.FillWidth(1.0)
+	[
+		SNew(SSpacer)
+	]
+
+	// View option
+	+SHorizontalBox::Slot()
+	.AutoWidth()
+	[
+		MakeViewOptionsWidget()
+	];
+}
+
+void FConcertSessionActivitiesOptions::OnOptionToggled(const FName CheckBoxId)
+{
+	if (CheckBoxId == ConcertSessionActivityUtils::DisplayRelativeTimeCheckBoxId)
+	{
+		bDisplayRelativeTime = !bDisplayRelativeTime;
+	}
+	else if (CheckBoxId == ConcertSessionActivityUtils::ShowConnectionActivitiesCheckBoxId)
+	{
+		bDisplayConnectionActivities = !bDisplayConnectionActivities;
+	}
+	else if (CheckBoxId == ConcertSessionActivityUtils::ShowLockActivitiesCheckBoxId)
+	{
+		bDisplayLockActivities = !bDisplayLockActivities;
+	}
+	else if (CheckBoxId == ConcertSessionActivityUtils::ShowPackageActivitiesCheckBoxId)
+	{
+		bDisplayPackageActivities = !bDisplayPackageActivities;
+	}
+	else if (CheckBoxId == ConcertSessionActivityUtils::ShowTransactionActivitiesCheckBoxId)
+	{
+		bDisplayTransactionActivities = !bDisplayTransactionActivities;
+	}
+	else if (CheckBoxId == ConcertSessionActivityUtils::ShowIgnoredActivitiesCheckBoxId)
+	{
+		bDisplayIgnoredActivities = !bDisplayIgnoredActivities;
+	}	
+}
+
+#undef LOCTEXT_NAMESPACE
diff --git a/Engine/Plugins/Developer/Concert/ConcertUI/ConcertSharedSlate/Source/ConcertSharedSlate/Private/Session/Browser/SConcertSessionBrowser.cpp b/Engine/Plugins/Developer/Concert/ConcertUI/ConcertSharedSlate/Source/ConcertSharedSlate/Private/Session/Browser/SConcertSessionBrowser.cpp
new file mode 100644
--- /dev/null
+++ b/Engine/Plugins/Developer/Concert/ConcertUI/ConcertSharedSlate/Source/ConcertSharedSlate/Private/Session/Browser/SConcertSessionBrowser.cpp
@@ -0,0 +1,1041 @@
+﻿// Copyright Epic Games, Inc. All Rights Reserved.
+
+#include "Session/Browser/SConcertSessionBrowser.h"
+
+#include "ConcertFrontendStyle.h"
+#include "Session/Browser/ConcertBrowserUtils.h"
+#include "Session/Browser/ConcertSessionBrowserSettings.h"
+#include "Session/Browser/ConcertSessionItem.h"
+#include "Session/Browser/IConcertSessionBrowserController.h"
+#include "Session/Browser/SNewSessionRow.h"
+#include "Session/Browser/SSaveRestoreSessionRow.h"
+#include "Session/Browser/SSessionRow.h"
+
+#include "Framework/MultiBox/MultiBoxBuilder.h"
+#include "Misc/TextFilter.h"
+#include "Styling/AppStyle.h"
+#include "Widgets/Images/SImage.h"
+#include "Widgets/Input/SEditableTextBox.h"
+#include "Widgets/Input/SSearchBox.h"
+#include "Widgets/Layout/SBorder.h"
+#include "Widgets/Layout/SSpacer.h"
+#include "Widgets/Layout/SSeparator.h"
+#include "Widgets/Layout/SUniformGridPanel.h"
+#include "Widgets/Text/STextBlock.h"
+#include "Widgets/Views/SListView.h"
+#include "Widgets/Views/STableRow.h"
+
+#define LOCTEXT_NAMESPACE "SConcertBrowser"
+
+const FName SConcertSessionBrowser::ControlButtonExtensionHooks::BeforeSeparator("BeforeSeparator");
+const FName SConcertSessionBrowser::ControlButtonExtensionHooks::Separator("Separator");
+const FName SConcertSessionBrowser::ControlButtonExtensionHooks::AfterSeparator("AfterSeparator");
+
+const FName SConcertSessionBrowser::SessionContextMenuExtensionHooks::ManageSession("ManageSession");
+
+void SConcertSessionBrowser::Construct(const FArguments& InArgs, TSharedRef<IConcertSessionBrowserController> InController, TSharedPtr<FText> InSearchText)
+{
+	Controller = InController;
+
+	// Reload the persistent settings, such as the filters.
+	PersistentSettings = GetMutableDefault<UConcertSessionBrowserSettings>();
+
+	DefaultSessionName = InArgs._DefaultSessionName;
+	DefaultServerUrl = InArgs._DefaultServerURL;
+
+	ExtendSessionContextMenu = InArgs._ExtendSessionContextMenu;
+	OnSessionClicked = InArgs._OnSessionClicked;
+	OnSessionDoubleClicked = InArgs._OnSessionDoubleClicked;
+	OnRequestedDeleteSession = InArgs._OnRequestedDeleteSession;
+	CanDeleteArchivedSession = InArgs._CanDeleteArchivedSession;
+	CanDeleteActiveSession = InArgs._CanDeleteActiveSession;
+
+	// Setup search filter.
+	SearchedText = InSearchText; // Reload a previous search text (in any). Useful to remember searched text between join/leave sessions, but not persistent if the tab is closed.
+	SearchTextFilter = MakeShared<TTextFilter<const FConcertSessionItem&>>(TTextFilter<const FConcertSessionItem&>::FItemToStringArray::CreateSP(this, &SConcertSessionBrowser::PopulateSearchStrings));
+	SearchTextFilter->OnChanged().AddSP(this, &SConcertSessionBrowser::RefreshSessionList);
+
+	ChildSlot
+	[
+		MakeBrowserContent(InArgs)
+	];
+
+	if (!SearchedText->IsEmpty())
+	{
+		SearchBox->SetText(*SearchedText); // This trigger the chain of actions to apply the search filter.
+	}
+}
+
+TSharedRef<SWidget> SConcertSessionBrowser::MakeBrowserContent(const FArguments& InArgs)
+{
+	TSharedRef<SWidget> SessionTable = MakeSessionTableView();
+	return SNew(SBorder)
+		.BorderImage(FAppStyle::Get().GetBrush("ToolPanel.GroupBorder"))
+		.Padding(FMargin(1.0f, 2.0f))
+		[
+			SNew(SVerticalBox)
+
+			+SVerticalBox::Slot()
+			.AutoHeight()
+			[
+				MakeControlBar(InArgs)
+			]
+
+			// Session list.
+			+SVerticalBox::Slot()
+			.FillHeight(1.0f)
+			.Padding(1.0f, 2.0f)
+			[
+				InArgs._ExtendSessionTable.IsBound() ? InArgs._ExtendSessionTable.Execute(SessionTable) : SessionTable
+			]
+
+			+SVerticalBox::Slot()
+			.AutoHeight()
+			.Padding(FMargin(2.0f, 0.0f))
+			[
+				SNew(SSeparator)
+			]
+
+			// Session Count/View options filter.
+			+SVerticalBox::Slot()
+			.AutoHeight()
+			.Padding(FMargin(2.0f, 0.0f))
+			[
+				MakeSessionViewOptionsBar()
+			]
+		];
+}
+
+void SConcertSessionBrowser::RefreshSessionList()
+{
+	// Remember the selected instances (if any).
+	TArray<TSharedPtr<FConcertSessionItem>> SelectedItems = SessionsView->GetSelectedItems();
+	TArray<TSharedPtr<FConcertSessionItem>> ReselectedItems;
+	TSharedPtr<FConcertSessionItem> NewEditableRowParent;
+
+	// Predicate returning true if the specified item should be re-selected.
+	auto IsSelected = [&SelectedItems](const FConcertSessionItem& Item)
+	{
+		return SelectedItems.ContainsByPredicate([Item](const TSharedPtr<FConcertSessionItem>& Visited) { return *Visited == Item; });
+	};
+
+	// Matches the object instances before the update to the new instance after the update.
+	auto ReconcileObjectInstances = [&](TSharedPtr<FConcertSessionItem> NewItem)
+	{
+		if (IsSelected(*NewItem))
+		{
+			ReselectedItems.Add(NewItem);
+		}
+		else if (EditableSessionRowParent.IsValid() && !NewEditableRowParent.IsValid() && *EditableSessionRowParent == *NewItem)
+		{
+			NewEditableRowParent = NewItem;
+		}
+	};
+
+	// Clear sessions.
+	Sessions.Reset();
+
+	// Populate the live sessions.
+	for (const IConcertSessionBrowserController::FActiveSessionInfo& ActiveSession : GetController()->GetActiveSessions())
+	{
+		FConcertSessionItem NewItem(FConcertSessionItem::EType::ActiveSession, ActiveSession.SessionInfo.SessionName, ActiveSession.SessionInfo.SessionId, ActiveSession.ServerInfo.ServerName, ActiveSession.ServerInfo.AdminEndpointId, ActiveSession.ServerInfo.ServerFlags);
+		if (!IsFilteredOut(NewItem))
+		{
+			Sessions.Emplace(MakeShared<FConcertSessionItem>(MoveTemp(NewItem)));
+			ReconcileObjectInstances(Sessions.Last());
+		}
+	}
+
+	// Populate the archived.
+	for (const IConcertSessionBrowserController::FArchivedSessionInfo& ArchivedSession : GetController()->GetArchivedSessions())
+	{
+		FConcertSessionItem NewItem(FConcertSessionItem::EType::ArchivedSession, ArchivedSession.SessionInfo.SessionName, ArchivedSession.SessionInfo.SessionId, ArchivedSession.ServerInfo.ServerName, ArchivedSession.ServerInfo.AdminEndpointId, ArchivedSession.ServerInfo.ServerFlags);
+		if (!IsFilteredOut(NewItem))
+		{
+			Sessions.Emplace(MakeShared<FConcertSessionItem>(MoveTemp(NewItem)));
+			ReconcileObjectInstances(Sessions.Last());
+		}
+	}
+
+	// Restore the editable row state. (SortSessionList() below will ensure the parent/child relationship)
+	EditableSessionRowParent = NewEditableRowParent;
+	if (EditableSessionRow.IsValid())
+	{
+		if (EditableSessionRow->Type == FConcertSessionItem::EType::NewSession)
+		{
+			Sessions.Insert(EditableSessionRow, 0); // Always put 'new session' row at the top.
+		}
+		else if (EditableSessionRowParent.IsValid())
+		{
+			Sessions.Add(EditableSessionRow); // SortSessionList() called below will ensure the correct parent/child order.
+		}
+	}
+
+	// Restore previous selection.
+	if (ReselectedItems.Num())
+	{
+		for (const TSharedPtr<FConcertSessionItem>& Item : ReselectedItems)
+		{
+			SessionsView->SetItemSelection(Item, true);
+		}
+	}
+
+	SortSessionList();
+	SessionsView->RequestListRefresh();
+}
+
+TSharedPtr<IConcertSessionBrowserController> SConcertSessionBrowser::GetController() const
+{
+	TSharedPtr<IConcertSessionBrowserController> Result = Controller.Pin();
+	check(Result);
+	return Result;
+}
+
+void SConcertSessionBrowser::OnSearchTextChanged(const FText& InFilterText)
+{
+	SearchTextFilter->SetRawFilterText(InFilterText);
+	SearchBox->SetError( SearchTextFilter->GetFilterErrorText() );
+	*SearchedText = InFilterText;
+
+	bRefreshSessionFilter = true;
+}
+
+void SConcertSessionBrowser::OnSearchTextCommitted(const FText& InFilterText, ETextCommit::Type CommitType)
+{
+	if (!InFilterText.EqualTo(*SearchedText))
+	{
+		OnSearchTextChanged(InFilterText);
+	}
+}
+
+void SConcertSessionBrowser::PopulateSearchStrings(const FConcertSessionItem& Item, TArray<FString>& OutSearchStrings) const
+{
+	OutSearchStrings.Add(Item.ServerName);
+	OutSearchStrings.Add(Item.SessionName);
+}
+
+bool SConcertSessionBrowser::IsFilteredOut(const FConcertSessionItem& Item) const
+{
+	bool bIsDefaultServer = LastDefaultServerUrl.IsEmpty() || Item.ServerName == LastDefaultServerUrl;
+
+	return (!PersistentSettings->bShowActiveSessions && (Item.Type == FConcertSessionItem::EType::ActiveSession || Item.Type == FConcertSessionItem::EType::SaveSession)) ||
+	       (!PersistentSettings->bShowArchivedSessions && (Item.Type == FConcertSessionItem::EType::ArchivedSession || Item.Type == FConcertSessionItem::EType::RestoreSession)) ||
+	       (PersistentSettings->bShowDefaultServerSessionsOnly && !bIsDefaultServer) ||
+	       !SearchTextFilter->PassesFilter(Item);
+}
+
+FText SConcertSessionBrowser::HighlightSearchText() const
+{
+	return *SearchedText;
+}
+
+TSharedRef<SWidget> SConcertSessionBrowser::MakeControlBar(const FArguments& InArgs)
+{
+	return SNew(SHorizontalBox)
+		// The New/Join/Restore/Delete/Archive buttons
+		+SHorizontalBox::Slot()
+		.AutoWidth()
+		[
+			MakeButtonBar(InArgs)
+		]
+
+		// The search text.
+		+SHorizontalBox::Slot()
+		.FillWidth(1.0)
+		.Padding(4.0f, 5.0f, 8.0f, 5.0f)
+		[
+			SAssignNew(SearchBox, SSearchBox)
+			.HintText(LOCTEXT("SearchHint", "Search Session"))
+			.OnTextChanged(this, &SConcertSessionBrowser::OnSearchTextChanged)
+			.OnTextCommitted(this, &SConcertSessionBrowser::OnSearchTextCommitted)
+			.DelayChangeNotificationsWhileTyping(true)
+		]
+
+		// Optional: everything to the right of the search bar, e.g. user name and settings combo button
+		+SHorizontalBox::Slot()
+		.VAlign(VAlign_Center)
+		.AutoWidth()
+		[
+			InArgs._RightOfSearchBar.Widget
+		];
+}
+
+TSharedRef<SWidget> SConcertSessionBrowser::MakeButtonBar(const FArguments& InArgs)
+{
+	TSharedRef<FExtender> Extender = MakeShared<FExtender>();
+	InArgs._ExtendControllButtons.ExecuteIfBound(Extender.Get());
+	FToolBarBuilder RowBuilder(nullptr, FMultiBoxCustomization::None, Extender);
+	
+	RowBuilder.BeginSection(ControlButtonExtensionHooks::BeforeSeparator);
+	// New Session
+	RowBuilder.AddWidget(
+		ConcertBrowserUtils::MakeIconButton(
+			FConcertFrontendStyle::Get()->GetBrush("Concert.NewSession"),
+			LOCTEXT("NewButtonTooltip", "Create a new session"),
+			TAttribute<bool>(this, &SConcertSessionBrowser::IsNewButtonEnabledInternal),
+			FOnClicked::CreateSP(this, &SConcertSessionBrowser::OnNewButtonClicked)
+		)
+	);
+	RowBuilder.EndSection();
+	
+	RowBuilder.AddSeparator(ControlButtonExtensionHooks::Separator);
+	RowBuilder.BeginSection(ControlButtonExtensionHooks::AfterSeparator);
+	
+	// Restore (Share the same slot as Join)
+	RowBuilder.AddWidget(
+		ConcertBrowserUtils::MakeIconButton(
+			FConcertFrontendStyle::Get()->GetBrush("Concert.RestoreSession"),
+			LOCTEXT("RestoreButtonTooltip", "Restore the selected session"),
+			TAttribute<bool>(this, &SConcertSessionBrowser::IsRestoreButtonEnabledInternal),
+			FOnClicked::CreateSP(this, &SConcertSessionBrowser::OnRestoreButtonClicked),
+			TAttribute<EVisibility>::Create([this]() { return IsRestoreButtonEnabledInternal() ? EVisibility::Visible : EVisibility::Collapsed; })
+			)
+	);
+	// Archive.
+	RowBuilder.AddWidget(
+		ConcertBrowserUtils::MakeIconButton(FConcertFrontendStyle::Get()->GetBrush("Concert.ArchiveSession"), LOCTEXT("ArchiveButtonTooltip", "Archive the selected session"),
+			TAttribute<bool>(this, &SConcertSessionBrowser::IsArchiveButtonEnabledInternal),
+			FOnClicked::CreateSP(this, &SConcertSessionBrowser::OnArchiveButtonClicked))
+	);
+	// Delete.
+	RowBuilder.AddWidget(
+		ConcertBrowserUtils::MakeIconButton(FConcertFrontendStyle::Get()->GetBrush("Concert.DeleteSession"), LOCTEXT("DeleteButtonTooltip", "Delete the selected session if permitted"),
+			TAttribute<bool>(this, &SConcertSessionBrowser::IsDeleteButtonEnabledInternal),
+			FOnClicked::CreateSP(this, &SConcertSessionBrowser::OnDeleteButtonClicked))
+	);
+	RowBuilder.EndSection();
+
+	return RowBuilder.MakeWidget();
+}
+
+TSharedRef<SWidget> SConcertSessionBrowser::MakeSessionTableView()
+{
+	PrimarySortedColumn = ConcertBrowserUtils::IconColName;
+	PrimarySortMode = EColumnSortMode::Ascending;
+	SecondarySortedColumn = ConcertBrowserUtils::SessionColName;
+	SecondarySortMode = EColumnSortMode::Ascending;
+
+	return SAssignNew(SessionsView, SListView<TSharedPtr<FConcertSessionItem>>)
+		.SelectionMode(ESelectionMode::Single)
+		.ListItemsSource(&Sessions)
+		.OnGenerateRow(this, &SConcertSessionBrowser::OnGenerateSessionRowWidget)
+		.SelectionMode(ESelectionMode::Single)
+		.OnSelectionChanged(this, &SConcertSessionBrowser::OnSessionSelectionChanged)
+		.OnContextMenuOpening(this, &SConcertSessionBrowser::MakeContextualMenu)
+		.HeaderRow(
+			SNew(SHeaderRow)
+
+			+SHeaderRow::Column(ConcertBrowserUtils::IconColName)
+			.DefaultLabel(FText::GetEmpty())
+			.SortPriority(this, &SConcertSessionBrowser::GetColumnSortPriority, ConcertBrowserUtils::IconColName)
+			.SortMode(this, &SConcertSessionBrowser::GetColumnSortMode, ConcertBrowserUtils::IconColName)
+			.OnSort(this, &SConcertSessionBrowser::OnColumnSortModeChanged)
+			.FixedWidth(20)
+
+			+SHeaderRow::Column(ConcertBrowserUtils::SessionColName)
+			.DefaultLabel(LOCTEXT("SessioName", "Session"))
+			.SortPriority(this, &SConcertSessionBrowser::GetColumnSortPriority, ConcertBrowserUtils::SessionColName)
+			.SortMode(this, &SConcertSessionBrowser::GetColumnSortMode, ConcertBrowserUtils::SessionColName)
+			.OnSort(this, &SConcertSessionBrowser::OnColumnSortModeChanged)
+
+			+SHeaderRow::Column(ConcertBrowserUtils::ServerColName)
+			.DefaultLabel(LOCTEXT("Server", "Server"))
+			.SortPriority(this, &SConcertSessionBrowser::GetColumnSortPriority, ConcertBrowserUtils::ServerColName)
+			.SortMode(this, &SConcertSessionBrowser::GetColumnSortMode, ConcertBrowserUtils::ServerColName)
+			.OnSort(this, &SConcertSessionBrowser::OnColumnSortModeChanged));
+}
+
+EColumnSortMode::Type SConcertSessionBrowser::GetColumnSortMode(const FName ColumnId) const
+{
+	if (ColumnId == PrimarySortedColumn)
+	{
+		return PrimarySortMode;
+	}
+	else if (ColumnId == SecondarySortedColumn)
+	{
+		return SecondarySortMode;
+	}
+
+	return EColumnSortMode::None;
+}
+
+EColumnSortPriority::Type SConcertSessionBrowser::GetColumnSortPriority(const FName ColumnId) const
+{
+	if (ColumnId == PrimarySortedColumn)
+	{
+		return EColumnSortPriority::Primary;
+	}
+	else if (ColumnId == SecondarySortedColumn)
+	{
+		return EColumnSortPriority::Secondary;
+	}
+
+	return EColumnSortPriority::Max; // No specific priority.
+}
+
+void SConcertSessionBrowser::OnColumnSortModeChanged(const EColumnSortPriority::Type SortPriority, const FName& ColumnId, const EColumnSortMode::Type InSortMode)
+{
+	if (SortPriority == EColumnSortPriority::Primary)
+	{
+		PrimarySortedColumn = ColumnId;
+		PrimarySortMode = InSortMode;
+
+		if (ColumnId == SecondarySortedColumn) // Cannot be primary and secondary at the same time.
+		{
+			SecondarySortedColumn = FName();
+			SecondarySortMode = EColumnSortMode::None;
+		}
+	}
+	else if (SortPriority == EColumnSortPriority::Secondary)
+	{
+		SecondarySortedColumn = ColumnId;
+		SecondarySortMode = InSortMode;
+	}
+
+	SortSessionList();
+	SessionsView->RequestListRefresh();
+}
+
+void SConcertSessionBrowser::SortSessionList()
+{
+	check(!PrimarySortedColumn.IsNone()); // Should always have a primary column. User cannot clear this one.
+
+	auto Compare = [this](const TSharedPtr<FConcertSessionItem>& Lhs, const TSharedPtr<FConcertSessionItem>& Rhs, const FName& ColName, EColumnSortMode::Type SortMode)
+	{
+		if (Lhs->Type == FConcertSessionItem::EType::NewSession) // Always keep editable 'new session' row at the top.
+		{
+			return true;
+		}
+		else if (Rhs->Type == FConcertSessionItem::EType::NewSession)
+		{
+			return false;
+		}
+
+		if (ColName == ConcertBrowserUtils::IconColName)
+		{
+			return SortMode == EColumnSortMode::Ascending ? Lhs->Type < Rhs->Type : Lhs->Type > Rhs->Type;
+		}
+		else if (ColName == ConcertBrowserUtils::SessionColName)
+		{
+			return SortMode == EColumnSortMode::Ascending ? Lhs->SessionName < Rhs->SessionName : Lhs->SessionName > Rhs->SessionName;
+		}
+		else
+		{
+			return SortMode == EColumnSortMode::Ascending ? Lhs->ServerName < Rhs->ServerName : Lhs->ServerName > Rhs->ServerName;
+		}
+	};
+
+	Sessions.Sort([&](const TSharedPtr<FConcertSessionItem>& Lhs, const TSharedPtr<FConcertSessionItem>& Rhs)
+	{
+		if (Compare(Lhs, Rhs, PrimarySortedColumn, PrimarySortMode))
+		{
+			return true; // Lhs must be before Rhs based on the primary sort order.
+		}
+		else if (Compare(Rhs, Lhs, PrimarySortedColumn, PrimarySortMode)) // Invert operands order (goal is to check if operands are equal or not)
+		{
+			return false; // Rhs must be before Lhs based on the primary sort.
+		}
+		else // Lhs == Rhs on the primary column, need to order accoding the secondary column if one is set.
+		{
+			return SecondarySortedColumn.IsNone() ? false : Compare(Lhs, Rhs, SecondarySortedColumn, SecondarySortMode);
+		}
+	});
+
+	EnsureEditableParentChildOrder();
+}
+
+void SConcertSessionBrowser::EnsureEditableParentChildOrder()
+{
+	// This is for Archiving or Restoring a session. We keep the editable row below the session to archive or restore and visually link them with small wires in UI.
+	if (EditableSessionRowParent.IsValid())
+	{
+		check(EditableSessionRow.IsValid());
+		Sessions.Remove(EditableSessionRow);
+
+		int32 ParentIndex = Sessions.IndexOfByKey(EditableSessionRowParent);
+		if (ParentIndex != INDEX_NONE)
+		{
+			Sessions.Insert(EditableSessionRow, ParentIndex + 1); // Insert the 'child' below its parent.
+		}
+	}
+}
+
+TSharedRef<ITableRow> SConcertSessionBrowser::OnGenerateSessionRowWidget(TSharedPtr<FConcertSessionItem> Item, const TSharedRef<STableViewBase>& OwnerTable)
+{
+	switch (Item->Type)
+	{
+		case FConcertSessionItem::EType::ActiveSession:
+			return MakeActiveSessionRowWidget(Item, OwnerTable);
+
+		case FConcertSessionItem::EType::ArchivedSession:
+			return MakeArchivedSessionRowWidget(Item, OwnerTable);
+
+		case FConcertSessionItem::EType::NewSession:
+			return MakeNewSessionRowWidget(Item, OwnerTable);
+
+		case FConcertSessionItem::EType::RestoreSession:
+			return MakeRestoreSessionRowWidget(Item, OwnerTable);
+
+		default:
+			check(Item->Type == FConcertSessionItem::EType::SaveSession);
+			return MakeSaveSessionRowWidget(Item, OwnerTable);
+	}
+}
+
+TSharedRef<ITableRow> SConcertSessionBrowser::MakeActiveSessionRowWidget(const TSharedPtr<FConcertSessionItem>& ActiveItem, const TSharedRef<STableViewBase>& OwnerTable)
+{
+	const TOptional<FConcertSessionInfo> SessionInfo = GetController()->GetActiveSessionInfo(ActiveItem->ServerAdminEndpointId, ActiveItem->SessionId);
+
+	// Add an 'Active Session' row. Clicking the row icon joins the session.
+	return SNew(SSessionRow, ActiveItem, OwnerTable)
+		.OnDoubleClickFunc([this](const TSharedPtr<FConcertSessionItem>& Item) { OnSessionDoubleClicked.ExecuteIfBound(Item); })
+		.OnRenameFunc([this](const TSharedPtr<FConcertSessionItem>& Item, const FString& NewName) { RequestRenameSession(Item, NewName); })
+		.IsDefaultSession([this](TSharedPtr<FConcertSessionItem> ItemPin)
+		{
+			return ItemPin->Type == FConcertSessionItem::EType::ActiveSession
+				&& DefaultSessionName.IsSet() && ItemPin->SessionName == DefaultSessionName.Get()
+				&& DefaultServerUrl.IsSet() && ItemPin->ServerName == DefaultServerUrl.Get();
+		})
+		.HighlightText(this, &SConcertSessionBrowser::HighlightSearchText)
+		.ToolTipText(SessionInfo ? SessionInfo->ToDisplayString() : FText::GetEmpty())
+		.IsSelected_Lambda([this, ActiveItem]() { return SessionsView->GetSelectedItems().Num() == 1 && SessionsView->GetSelectedItems()[0] == ActiveItem; });
+}
+
+TSharedRef<ITableRow> SConcertSessionBrowser::MakeArchivedSessionRowWidget(const TSharedPtr<FConcertSessionItem>& ArchivedItem, const TSharedRef<STableViewBase>& OwnerTable)
+{
+	const TOptional<FConcertSessionInfo> SessionInfo = GetController()->GetArchivedSessionInfo(ArchivedItem->ServerAdminEndpointId, ArchivedItem->SessionId);
+
+	// Add an 'Archived Session' row. Clicking the row icon adds a 'Restore as' row to the table.
+	return SNew(SSessionRow, ArchivedItem, OwnerTable)
+		.OnDoubleClickFunc([this](const TSharedPtr<FConcertSessionItem>& Item) { InsertRestoreSessionAsEditableRowInternal(Item); })
+		.OnRenameFunc([this](const TSharedPtr<FConcertSessionItem>& Item, const FString& NewName) { RequestRenameSession(Item, NewName); })
+		.IsDefaultSession([this](TSharedPtr<FConcertSessionItem> ItemPin)
+		{
+			return ItemPin->Type == FConcertSessionItem::EType::ActiveSession
+				&& DefaultSessionName.IsSet() && ItemPin->SessionName == DefaultSessionName.Get()
+				&& DefaultServerUrl.IsSet() && ItemPin->ServerName == DefaultServerUrl.Get();
+		})
+		.HighlightText(this, &SConcertSessionBrowser::HighlightSearchText)
+		.ToolTipText(SessionInfo ? SessionInfo->ToDisplayString() : FText::GetEmpty())
+		.IsSelected_Lambda([this, ArchivedItem]() { return SessionsView->GetSelectedItems().Num() == 1 && SessionsView->GetSelectedItems()[0] == ArchivedItem; });
+}
+
+TSharedRef<ITableRow> SConcertSessionBrowser::MakeNewSessionRowWidget(const TSharedPtr<FConcertSessionItem>& NewItem, const TSharedRef<STableViewBase>& OwnerTable)
+{
+	// Add an editable 'New Session' row in the table to let user pick a name and a server.
+	TAttribute<FString> DefaultServerUrlArg = DefaultServerUrl.IsSet()
+		? TAttribute<FString>::Create([this](){ return DefaultServerUrl.Get(); }) : TAttribute<FString>();
+	return SNew(SNewSessionRow, NewItem, OwnerTable)
+		.GetServerFunc([this]() { return GetController()->GetServers(); }) // Let the row pull the servers for the combo box.
+		.OnAcceptFunc([this](const TSharedPtr<FConcertSessionItem>& Item) { RequestCreateSession(Item); }) // Accepting creates the session.
+		.OnDeclineFunc([this](const TSharedPtr<FConcertSessionItem>& Item) { RemoveSessionRow(Item); })  // Declining removes the editable 'new session' row from the view.
+		.HighlightText(this, &SConcertSessionBrowser::HighlightSearchText)
+		.DefaultServerURL(DefaultServerUrlArg);
+}
+
+TSharedRef<ITableRow> SConcertSessionBrowser::MakeSaveSessionRowWidget(const TSharedPtr<FConcertSessionItem>& SaveItem, const TSharedRef<STableViewBase>& OwnerTable)
+{
+	// Add an editable 'Save Session' row in the table to let the user enter an archive name.
+	return SNew(SSaveRestoreSessionRow, SaveItem, OwnerTable)
+		.OnAcceptFunc([this](const TSharedPtr<FConcertSessionItem>& Item, const FString& ArchiveName) { RequestArchiveSession(Item, ArchiveName); }) // Accepting archive the session.
+		.OnDeclineFunc([this](const TSharedPtr<FConcertSessionItem>& Item) { RemoveSessionRow(Item); }) // Declining removes the editable 'save session as' row from the view.
+		.HighlightText(this, &SConcertSessionBrowser::HighlightSearchText);
+}
+
+TSharedRef<ITableRow> SConcertSessionBrowser::MakeRestoreSessionRowWidget(const TSharedPtr<FConcertSessionItem>& RestoreItem, const TSharedRef<STableViewBase>& OwnerTable)
+{
+	// Add an editable 'Restore Session' row in the table to let the user enter a session name.
+	return SNew(SSaveRestoreSessionRow, RestoreItem, OwnerTable)
+		.OnAcceptFunc([this](const TSharedPtr<FConcertSessionItem>& Item, const FString& SessionName) { RequestRestoreSession(Item, SessionName); }) // Accepting restores the session.
+		.OnDeclineFunc([this](const TSharedPtr<FConcertSessionItem>& Item) { RemoveSessionRow(Item); }) // Declining removes the editable 'restore session as' row from the view.
+		.HighlightText(this, &SConcertSessionBrowser::HighlightSearchText);
+}
+
+void SConcertSessionBrowser::InsertNewSessionEditableRowInternal()
+{
+	// Insert a 'new session' editable row.
+	InsertEditableSessionRow(MakeShared<FConcertSessionItem>(FConcertSessionItem::EType::NewSession, FString(), FGuid(), FString(), FGuid(), EConcertServerFlags::None), nullptr);
+}
+
+void SConcertSessionBrowser::InsertRestoreSessionAsEditableRowInternal(const TSharedPtr<FConcertSessionItem>& ArchivedItem)
+{
+	// Insert the 'restore session as ' editable row just below the 'archived' item to restore.
+	InsertEditableSessionRow(MakeShared<FConcertSessionItem>(FConcertSessionItem::EType::RestoreSession, ArchivedItem->SessionName, ArchivedItem->SessionId, ArchivedItem->ServerName, ArchivedItem->ServerAdminEndpointId, ArchivedItem->ServerFlags), ArchivedItem);
+}
+
+void SConcertSessionBrowser::InsertArchiveSessionAsEditableRow(const TSharedPtr<FConcertSessionItem>& LiveItem)
+{
+	// Insert the 'save session as' editable row just below the 'active' item to save.
+	InsertEditableSessionRow(MakeShared<FConcertSessionItem>(FConcertSessionItem::EType::SaveSession, LiveItem->SessionName, LiveItem->SessionId, LiveItem->ServerName, LiveItem->ServerAdminEndpointId, LiveItem->ServerFlags), LiveItem);
+}
+
+void SConcertSessionBrowser::InsertEditableSessionRow(TSharedPtr<FConcertSessionItem> EditableItem, TSharedPtr<FConcertSessionItem> ParentItem)
+{
+	// Insert the new row below its parent (if any).
+	int32 ParentIndex = Sessions.IndexOfByKey(ParentItem);
+	Sessions.Insert(EditableItem, ParentIndex != INDEX_NONE ? ParentIndex + 1 : 0);
+
+	// Ensure there is only one editable row at the time, removing the row being edited (if any).
+	Sessions.Remove(EditableSessionRow);
+	EditableSessionRow = EditableItem;
+	EditableSessionRowParent = ParentItem;
+
+	// Ensure the editable row added is selected and visible.
+	SessionsView->SetSelection(EditableItem);
+	SessionsView->RequestListRefresh();
+
+	// NOTE: Ideally, I would only use RequestScrollIntoView() to scroll the new item into view, but it did not work. If an item was added into an hidden part,
+	//       it was not always scrolled correctly into view. RequestNavigateToItem() worked much better, except when inserting the very first row in the list, in
+	//       such case calling the function would give the focus to the list view (showing a white dashed line around it).
+	if (ParentIndex == INDEX_NONE)
+	{
+		SessionsView->ScrollToTop(); // Item is inserted at 0. (New session)
+	}
+	else
+	{
+		SessionsView->RequestNavigateToItem(EditableItem);
+	}
+}
+
+void SConcertSessionBrowser::RemoveSessionRow(const TSharedPtr<FConcertSessionItem>& Item)
+{
+	Sessions.Remove(Item);
+
+	// Don't keep the editable row if its 'parent' is removed. (if the session to restore or archive gets deleted in the meantime)
+	if (Item == EditableSessionRowParent)
+	{
+		Sessions.Remove(EditableSessionRow);
+		EditableSessionRow.Reset();
+	}
+
+	// Clear the editable row state if its the one removed.
+	if (Item == EditableSessionRow)
+	{
+		EditableSessionRow.Reset();
+		EditableSessionRowParent.Reset();
+	}
+
+	SessionsView->RequestListRefresh();
+}
+
+TSharedRef<SWidget> SConcertSessionBrowser::MakeSessionViewOptionsBar()
+{
+	auto AddFilterMenu = [this]()
+	{
+		FMenuBuilder MenuBuilder(/*bInShouldCloseWindowAfterMenuSelection=*/true, nullptr);
+
+		MenuBuilder.AddMenuEntry(
+			LOCTEXT("ActiveSessions_Label", "Active Sessions"),
+			LOCTEXT("ActiveSessions_Tooltip", "Displays Active Sessions"),
+			FSlateIcon(),
+			FUIAction(
+				FExecuteAction::CreateSP(this, &SConcertSessionBrowser::OnFilterMenuChecked, ConcertBrowserUtils::ActiveSessionsCheckBoxMenuName),
+				FCanExecuteAction::CreateLambda([] { return true; }),
+				FIsActionChecked::CreateLambda([this] { return PersistentSettings->bShowActiveSessions; })),
+			NAME_None,
+			EUserInterfaceActionType::ToggleButton
+		);
+
+		MenuBuilder.AddMenuEntry(
+			LOCTEXT("ArchivedSessions_Label", "Archived Sessions"),
+			LOCTEXT("ArchivedSessions_Tooltip", "Displays Archived Sessions"),
+			FSlateIcon(),
+			FUIAction(
+				FExecuteAction::CreateSP(this, &SConcertSessionBrowser::OnFilterMenuChecked, ConcertBrowserUtils::ArchivedSessionsCheckBoxMenuName),
+				FCanExecuteAction::CreateLambda([] { return true; }),
+				FIsActionChecked::CreateLambda([this] { return PersistentSettings->bShowArchivedSessions; })),
+			NAME_None,
+			EUserInterfaceActionType::ToggleButton
+		);
+
+		MenuBuilder.AddMenuEntry(
+			LOCTEXT("DefaultServer_Label", "Default Server Sessions"),
+			LOCTEXT("DefaultServer_Tooltip", "Displays Sessions Hosted By the Default Server"),
+			FSlateIcon(),
+			FUIAction(
+				FExecuteAction::CreateSP(this, &SConcertSessionBrowser::OnFilterMenuChecked, ConcertBrowserUtils::DefaultServerCheckBoxMenuName),
+				FCanExecuteAction::CreateLambda([] { return true; }),
+				FIsActionChecked::CreateLambda([this] { return PersistentSettings->bShowDefaultServerSessionsOnly; })),
+			NAME_None,
+			EUserInterfaceActionType::ToggleButton
+		);
+
+		return MenuBuilder.MakeWidget();
+	};
+
+	return SNew(SHorizontalBox)
+		+SHorizontalBox::Slot()
+		.AutoWidth()
+		.VAlign(VAlign_Center)
+		[
+			SNew(STextBlock).Text_Lambda([this]()
+			{
+				// Don't count the 'New Session', 'Restore Session' and 'Archive Session' editable row, they are transient rows used for inline input only.
+				int32 DisplayedSessionNum = Sessions.Num() - (EditableSessionRow.IsValid() ? 1 : 0);
+				int32 AvailableSessionNum = GetController()->GetActiveSessions().Num() + GetController()->GetArchivedSessions().Num();
+				int32 ServerNum = GetController()->GetServers().Num();
+
+				// If all discovered session are displayed (none excluded by a filter).
+				if (DisplayedSessionNum == AvailableSessionNum)
+				{
+					if (GetController()->GetServers().Num() == 0)
+					{
+						return LOCTEXT("NoServerNoFilter", "No servers found");
+					}
+					else
+					{
+						return FText::Format(LOCTEXT("NSessionNServerNoFilter", "{0} {0}|plural(one=session,other=sessions) on {1} {1}|plural(one=server,other=servers)"),
+							DisplayedSessionNum, ServerNum);
+					}
+				}
+				else // A filter is excluding at least one session.
+				{
+					if (DisplayedSessionNum == 0)
+					{
+						return FText::Format(LOCTEXT("NoSessionMatchNServer", "No matching sessions ({0} total on {1} {1}|plural(one=server,other=servers))"),
+							AvailableSessionNum, ServerNum);
+					}
+					else
+					{
+						return FText::Format(LOCTEXT("NSessionNServer", "Showing {0} of {1} {1}|plural(one=session,other=sessions) on {2} {2}|plural(one=server,other=servers)"),
+							DisplayedSessionNum, AvailableSessionNum, ServerNum);
+					}
+				}
+			})
+		]
+
+		+SHorizontalBox::Slot()
+		.FillWidth(1.0)
+		[
+			SNew(SSpacer)
+		]
+
+		+SHorizontalBox::Slot()
+		.AutoWidth()
+		[
+			SNew(SComboButton)
+			.ComboButtonStyle(FAppStyle::Get(), "GenericFilters.ComboButtonStyle")
+			.ForegroundColor(FLinearColor::White)
+			.ContentPadding(0)
+			.OnGetMenuContent_Lambda(AddFilterMenu)
+			.HasDownArrow(true)
+			.ContentPadding(FMargin(1, 0))
+			.ButtonContent()
+			[
+				SNew(SHorizontalBox)
+
+				+SHorizontalBox::Slot()
+				.AutoWidth()
+				.VAlign(VAlign_Center)
+				[
+					SNew(SImage).Image(FAppStyle::Get().GetBrush("GenericViewButton")) // The eye ball image.
+				]
+
+				+SHorizontalBox::Slot()
+				.AutoWidth()
+				.Padding(2, 0, 0, 0)
+				.VAlign(VAlign_Center)
+				[
+					SNew(STextBlock).Text(LOCTEXT("ViewOptions", "View Options"))
+				]
+			]
+		];
+}
+
+void SConcertSessionBrowser::OnFilterMenuChecked(const FName MenuName)
+{
+	if (MenuName == ConcertBrowserUtils::ActiveSessionsCheckBoxMenuName)
+	{
+		PersistentSettings->bShowActiveSessions = !PersistentSettings->bShowActiveSessions;
+	}
+	else if (MenuName == ConcertBrowserUtils::ArchivedSessionsCheckBoxMenuName)
+	{
+		PersistentSettings->bShowArchivedSessions = !PersistentSettings->bShowArchivedSessions;
+	}
+	else if (MenuName == ConcertBrowserUtils::DefaultServerCheckBoxMenuName)
+	{
+		PersistentSettings->bShowDefaultServerSessionsOnly = !PersistentSettings->bShowDefaultServerSessionsOnly;
+	}
+	bRefreshSessionFilter = true;
+
+	PersistentSettings->SaveConfig();
+}
+
+TSharedPtr<SWidget> SConcertSessionBrowser::MakeContextualMenu()
+{
+	TArray<TSharedPtr<FConcertSessionItem>> SelectedItems = SessionsView->GetSelectedItems();
+	if (SelectedItems.Num() == 0 || (SelectedItems[0]->Type != FConcertSessionItem::EType::ActiveSession && SelectedItems[0]->Type != FConcertSessionItem::EType::ArchivedSession))
+	{
+		return nullptr; // No menu for editable rows.
+	}
+
+	TSharedPtr<FConcertSessionItem> Item = SelectedItems[0];
+
+	TSharedRef<FExtender> Extender = MakeShared<FExtender>();
+	ExtendSessionContextMenu.ExecuteIfBound(Item, Extender.Get());
+	FMenuBuilder MenuBuilder(/*bInShouldCloseWindowAfterMenuSelection=*/true, nullptr, Extender);
+
+	// Section title.
+	MenuBuilder.BeginSection(SessionContextMenuExtensionHooks::ManageSession, Item->Type == FConcertSessionItem::EType::ActiveSession ?
+		LOCTEXT("ActiveSessionSection", "Active Session"):
+		LOCTEXT("ArchivedSessionSection", "Archived Session"));
+
+	if (Item->Type == FConcertSessionItem::EType::ActiveSession)
+	{
+		MenuBuilder.AddMenuEntry(
+			LOCTEXT("CtxMenuArchive", "Archive"),
+			LOCTEXT("CtxMenuArchive_Tooltip", "Archived the Session"),
+			FSlateIcon(),
+			FUIAction(
+				FExecuteAction::CreateLambda([this](){ OnArchiveButtonClicked(); }),
+				FCanExecuteAction::CreateLambda([SelectedCount = SelectedItems.Num()] { return SelectedCount == 1; }),
+				FIsActionChecked::CreateLambda([this] { return false; })),
+			NAME_None,
+			EUserInterfaceActionType::Button);
+	}
+	else // Archive
+	{
+		MenuBuilder.AddMenuEntry(
+			LOCTEXT("CtxMenuRestore", "Restore"),
+			LOCTEXT("CtxMenuRestore_Tooltip", "Restore the Session"),
+			FSlateIcon(),
+			FUIAction(
+				FExecuteAction::CreateLambda([this](){ OnRestoreButtonClicked(); }),
+				FCanExecuteAction::CreateLambda([SelectedCount = SelectedItems.Num()] { return SelectedCount == 1; }),
+				FIsActionChecked::CreateLambda([this] { return false; })),
+			NAME_None,
+			EUserInterfaceActionType::Button);
+	}
+
+	MenuBuilder.AddMenuEntry(
+		LOCTEXT("CtxMenuRename", "Rename"),
+		LOCTEXT("CtxMenuRename_Tooltip", "Rename the Session"),
+		FSlateIcon(),
+		FUIAction(
+			FExecuteAction::CreateLambda([this, Item](){ OnBeginEditingSessionName(Item); }),
+			FCanExecuteAction::CreateLambda([this] { return IsRenameButtonEnabledInternal(); }),
+			FIsActionChecked::CreateLambda([this] { return false; })),
+		NAME_None,
+		EUserInterfaceActionType::Button);
+
+	MenuBuilder.AddMenuEntry(
+		LOCTEXT("CtxMenuDelete", "Delete"),
+		FText::Format(LOCTEXT("CtxMenuDelete_Tooltip", "Delete the {0}|plural(one=Session,other=Sessions)"), SelectedItems.Num()),
+		FSlateIcon(),
+		FUIAction(
+			FExecuteAction::CreateLambda([this](){ OnDeleteButtonClicked(); }),
+			FCanExecuteAction::CreateLambda([this] { return IsDeleteButtonEnabledInternal(); }),
+			FIsActionChecked::CreateLambda([this] { return false; })),
+		NAME_None,
+		EUserInterfaceActionType::Button);
+
+	MenuBuilder.EndSection();
+
+	return MenuBuilder.MakeWidget();
+}
+
+void SConcertSessionBrowser::OnSessionSelectionChanged(TSharedPtr<FConcertSessionItem> SelectedSession, ESelectInfo::Type SelectInfo)
+{
+	// Cancel editing the row to create, archive or restore a session (if any), unless the row was selected in code.
+	if (EditableSessionRow.IsValid() && SelectInfo != ESelectInfo::Direct)
+	{
+		check (SelectedSession != EditableSessionRow); // User should not be able to reselect an editable row as we remove it as soon as it is unselected.
+		RemoveSessionRow(EditableSessionRow);
+		check(!EditableSessionRow.IsValid() && !EditableSessionRowParent.IsValid()); // Expect to be cleared by RemoveSessionRow().
+	}
+
+	// Clear the list of clients (if any)
+	Clients.Reset();
+
+	OnSessionClicked.ExecuteIfBound(SelectedSession);
+}
+
+bool SConcertSessionBrowser::IsNewButtonEnabledInternal() const
+{
+	return GetController()->GetServers().Num() > 0;
+}
+
+bool SConcertSessionBrowser::IsRestoreButtonEnabledInternal() const
+{
+	TArray<TSharedPtr<FConcertSessionItem>> SelectedItems = SessionsView->GetSelectedItems();
+	return SelectedItems.Num() == 1 && SelectedItems[0]->Type == FConcertSessionItem::EType::ArchivedSession;
+}
+
+bool SConcertSessionBrowser::IsArchiveButtonEnabledInternal() const
+{
+	TArray<TSharedPtr<FConcertSessionItem>> SelectedItems = SessionsView->GetSelectedItems();
+	return SelectedItems.Num() == 1 && SelectedItems[0]->Type == FConcertSessionItem::EType::ActiveSession;
+}
+
+bool SConcertSessionBrowser::IsRenameButtonEnabledInternal() const
+{
+	TArray<TSharedPtr<FConcertSessionItem>> SelectedItems = SessionsView->GetSelectedItems();
+	if (SelectedItems.Num() != 1)
+	{
+		return false;
+	}
+
+	return (SelectedItems[0]->Type == FConcertSessionItem::EType::ActiveSession && GetController()->CanRenameActiveSession(SelectedItems[0]->ServerAdminEndpointId, SelectedItems[0]->SessionId)) ||
+	       (SelectedItems[0]->Type == FConcertSessionItem::EType::ArchivedSession && GetController()->CanRenameArchivedSession(SelectedItems[0]->ServerAdminEndpointId, SelectedItems[0]->SessionId));
+}
+
+bool SConcertSessionBrowser::IsDeleteButtonEnabledInternal() const
+{
+	TArray<TSharedPtr<FConcertSessionItem>> SelectedItems = SessionsView->GetSelectedItems();
+	if (SelectedItems.Num() == 0)
+	{
+		return false;
+	}
+
+	return (SelectedItems[0]->Type == FConcertSessionItem::EType::ActiveSession && GetController()->CanDeleteActiveSession(SelectedItems[0]->ServerAdminEndpointId, SelectedItems[0]->SessionId)) ||
+	       (SelectedItems[0]->Type == FConcertSessionItem::EType::ArchivedSession && GetController()->CanDeleteArchivedSession(SelectedItems[0]->ServerAdminEndpointId, SelectedItems[0]->SessionId));
+}
+
+FReply SConcertSessionBrowser::OnNewButtonClicked()
+{
+	InsertNewSessionEditableRowInternal();
+	return FReply::Handled();
+}
+
+FReply SConcertSessionBrowser::OnRestoreButtonClicked()
+{
+	TArray<TSharedPtr<FConcertSessionItem>> SelectedItems = SessionsView->GetSelectedItems();
+	if (SelectedItems.Num() == 1)
+	{
+		InsertRestoreSessionAsEditableRowInternal(SelectedItems[0]);
+	}
+
+	return FReply::Handled();
+}
+
+FReply SConcertSessionBrowser::OnArchiveButtonClicked()
+{
+	TArray<TSharedPtr<FConcertSessionItem>> SelectedItems = SessionsView->GetSelectedItems();
+	if (SelectedItems.Num() == 1)
+	{
+		InsertArchiveSessionAsEditableRow(SelectedItems[0]);
+	}
+
+	return FReply::Handled();
+}
+
+FReply SConcertSessionBrowser::OnDeleteButtonClicked()
+{
+	TArray<TSharedPtr<FConcertSessionItem>> SelectedItems = SessionsView->GetSelectedItems();
+	for (const TSharedPtr<FConcertSessionItem>& Item : SelectedItems)
+	{
+		RequestDeleteSession(Item);
+	}
+
+	return FReply::Handled();
+}
+
+void SConcertSessionBrowser::OnBeginEditingSessionName(TSharedPtr<FConcertSessionItem> Item)
+{
+	Item->OnBeginEditSessionNameRequest.Broadcast(); // Signal the row widget to enter in edit mode.
+}
+
+void SConcertSessionBrowser::Tick(const FGeometry& AllottedGeometry, const double InCurrentTime, const float InDeltaTime)
+{
+	// Ensure the 'default server' filter is updated when the configuration of the default server changes.
+	if (DefaultServerUrl.IsSet() && LastDefaultServerUrl != DefaultServerUrl.Get())
+	{
+		LastDefaultServerUrl = DefaultServerUrl.Get();
+		bRefreshSessionFilter = true;
+	}
+
+	// Should refresh the session filter?
+	if (bRefreshSessionFilter)
+	{
+		RefreshSessionList();
+		bRefreshSessionFilter = false;
+	}
+}
+
+FReply SConcertSessionBrowser::OnKeyDown(const FGeometry& MyGeometry, const FKeyEvent& InKeyEvent)
+{
+	// NOTE: When an 'editable row' text box has the focus the keys are grabbed by the text box but
+	//       if the editable row is still selected, but the text field doesn't have the focus anymore
+	//       the keys will end up here if the browser has the focus.
+
+	if (InKeyEvent.GetKey() == EKeys::Delete && !EditableSessionRow.IsValid()) // Delete selected row(s) unless the selected row is an 'editable' one.
+	{
+		for (TSharedPtr<FConcertSessionItem>& Item : SessionsView->GetSelectedItems())
+		{
+			RequestDeleteSession(Item);
+		}
+		return FReply::Handled();
+	}
+	else if (InKeyEvent.GetKey() == EKeys::Escape && EditableSessionRow.IsValid()) // Cancel 'new session', 'archive session' or 'restore session' action.
+	{
+		RemoveSessionRow(EditableSessionRow);
+		check(!EditableSessionRow.IsValid() && !EditableSessionRowParent.IsValid()); // Expect to be cleared by RemoveSessionRow().
+		return FReply::Handled();
+	}
+	else if (InKeyEvent.GetKey() == EKeys::F2 && !EditableSessionRow.IsValid())
+	{
+		TArray<TSharedPtr<FConcertSessionItem>> SelectedItems = SessionsView->GetSelectedItems();
+		if (SelectedItems.Num() == 1)
+		{
+			SelectedItems[0]->OnBeginEditSessionNameRequest.Broadcast(); // Broadcast the request.
+		}
+	}
+
+	return FReply::Unhandled();
+}
+
+void SConcertSessionBrowser::RequestCreateSession(const TSharedPtr<FConcertSessionItem>& NewItem)
+{
+	GetController()->CreateSession(NewItem->ServerAdminEndpointId, NewItem->SessionName);
+	RemoveSessionRow(NewItem); // The row used to edit the session name and pick the server.
+}
+
+void SConcertSessionBrowser::RequestArchiveSession(const TSharedPtr<FConcertSessionItem>& SaveItem, const FString& ArchiveName)
+{
+	GetController()->ArchiveSession(SaveItem->ServerAdminEndpointId, SaveItem->SessionId, ArchiveName, FConcertSessionFilter());
+	RemoveSessionRow(SaveItem); // The row used to edit the archive name.
+}
+
+void SConcertSessionBrowser::RequestRestoreSession(const TSharedPtr<FConcertSessionItem>& RestoreItem, const FString& SessionName)
+{
+	GetController()->RestoreSession(RestoreItem->ServerAdminEndpointId, RestoreItem->SessionId, SessionName, FConcertSessionFilter());
+	RemoveSessionRow(RestoreItem); // The row used to edit the restore as name.
+}
+
+void SConcertSessionBrowser::RequestRenameSession(const TSharedPtr<FConcertSessionItem>& RenamedItem, const FString& NewName)
+{
+	if (RenamedItem->Type == FConcertSessionItem::EType::ActiveSession)
+	{
+		GetController()->RenameActiveSession(RenamedItem->ServerAdminEndpointId, RenamedItem->SessionId, NewName);
+	}
+	else if (RenamedItem->Type == FConcertSessionItem::EType::ArchivedSession)
+	{
+		GetController()->RenameArchivedSession(RenamedItem->ServerAdminEndpointId, RenamedItem->SessionId, NewName);
+	}
+
+	// Display the new name until the server response is received. If the server refuses the new name, the discovery will reset the
+	// name (like if another client renamed it back) and the user will get a toast saying the rename failed.
+	RenamedItem->SessionName = NewName;
+}
+
+void SConcertSessionBrowser::RequestDeleteSession(const TSharedPtr<FConcertSessionItem>& DeletedItem)
+{
+	if (DeletedItem->Type == FConcertSessionItem::EType::ActiveSession)
+	{
+		if (!CanDeleteActiveSession.IsBound() || CanDeleteActiveSession.Execute(DeletedItem))
+		{
+			GetController()->DeleteActiveSession(DeletedItem->ServerAdminEndpointId, DeletedItem->SessionId);
+		}
+	}
+	else if (DeletedItem->Type == FConcertSessionItem::EType::ArchivedSession)
+	{
+		if (!CanDeleteArchivedSession.IsBound() || CanDeleteArchivedSession.Execute(DeletedItem))
+		{
+			GetController()->DeleteArchivedSession(DeletedItem->ServerAdminEndpointId, DeletedItem->SessionId);
+		}
+	}
+	
+	OnRequestedDeleteSession.ExecuteIfBound(DeletedItem);
+}
+
+#undef LOCTEXT_NAMESPACE
\ No newline at end of file
diff --git a/Engine/Plugins/Developer/Concert/ConcertUI/ConcertSharedSlate/Source/ConcertSharedSlate/Private/Session/Browser/SNewSessionRow.cpp b/Engine/Plugins/Developer/Concert/ConcertUI/ConcertSharedSlate/Source/ConcertSharedSlate/Private/Session/Browser/SNewSessionRow.cpp
new file mode 100644
--- /dev/null
+++ b/Engine/Plugins/Developer/Concert/ConcertUI/ConcertSharedSlate/Source/ConcertSharedSlate/Private/Session/Browser/SNewSessionRow.cpp
@@ -0,0 +1,348 @@
+﻿// Copyright Epic Games, Inc. All Rights Reserved.
+
+#include "SNewSessionRow.h"
+
+#include "Session/Browser/ConcertBrowserUtils.h"
+#include "Session/Browser/ConcertSessionItem.h"
+
+#include "EditorFontGlyphs.h"
+#include "EditorStyleSet.h"
+#include "Framework/Application/SlateApplication.h"
+#include "Styling/AppStyle.h"
+#include "Widgets/Input/SEditableTextBox.h"
+#include "Widgets/Layout/SBox.h"
+#include "Widgets/SBoxPanel.h"
+#include "Widgets/Layout/SUniformGridPanel.h"
+#include "Widgets/Text/STextBlock.h"
+#include "Widgets/Views/STableRow.h"
+
+#define LOCTEXT_NAMESPACE "SConcertBrowser"
+
+void SNewSessionRow::Construct(const FArguments& InArgs, TSharedPtr<FConcertSessionItem> InItem, const TSharedRef<STableViewBase>& InOwnerTableView)
+{
+	Item = MoveTemp(InItem);
+	GetServersFunc = InArgs._GetServerFunc;
+	AcceptFunc = InArgs._OnAcceptFunc;
+	DeclineFunc = InArgs._OnDeclineFunc;
+	HighlightText = InArgs._HighlightText;
+	DefaultServerURL = InArgs._DefaultServerURL;
+
+	GetServersFunc.CheckCallable();
+	AcceptFunc.CheckCallable();
+	DeclineFunc.CheckCallable();
+
+	// Construct base class
+	SMultiColumnTableRow<TSharedPtr<FConcertSessionItem>>::Construct(FSuperRowType::FArguments(), InOwnerTableView);
+
+	// Fill the server combo.
+	UpdateServerList();
+}
+
+void SNewSessionRow::Tick(const FGeometry& AllottedGeometry, const double InCurrentTime, const float InDeltaTime)
+{
+	// Performance could be improved by only updating it when the server list has changed but it's no biggy...
+	UpdateServerList();
+
+	// Should give the focus to an editable text.
+	if (!bInitialFocusTaken)
+	{
+		bInitialFocusTaken = FSlateApplication::Get().SetKeyboardFocus(EditableSessionName.ToSharedRef());
+	}
+}
+
+TSharedRef<SWidget> SNewSessionRow::GenerateWidgetForColumn(const FName& ColumnName)
+{
+	TSharedPtr<FConcertSessionItem> ItemPin = Item.Pin();
+
+	if (ColumnName == ConcertBrowserUtils::IconColName)
+	{
+		// 'New' icon
+		return SNew(SBox)
+			.VAlign(VAlign_Center)
+			.HAlign(HAlign_Center)
+			[
+				SNew(STextBlock)
+					.Font(FAppStyle::Get().GetFontStyle(ConcertBrowserUtils::IconColumnFontName))
+					.Text(FEditorFontGlyphs::Plus_Circle)
+			];
+	}
+	else if (ColumnName == ConcertBrowserUtils::SessionColName)
+	{
+		return SNew(SBox)
+			.VAlign(VAlign_Center)
+			.Padding(FMargin(0, 0, 2, 0))
+			[
+				SAssignNew(EditableSessionName, SEditableTextBox)
+				.HintText(LOCTEXT("EnterSessionNameHint", "Enter a session name"))
+				.OnTextCommitted(this, &SNewSessionRow::OnSessionNameCommitted)
+				.OnKeyDownHandler(this, &SNewSessionRow::OnKeyDownHandler)
+				.OnTextChanged(this, &SNewSessionRow::OnSessionNameChanged)
+			];
+	}
+	else
+	{
+		return SNew(SHorizontalBox)
+
+			// 'Server' combo
+			+SHorizontalBox::Slot()
+			.AutoWidth()
+			.Padding(0, 1)
+			[
+				SAssignNew(ServersComboBox, SComboBox<TSharedPtr<FConcertServerInfo>>)
+				.OptionsSource(&Servers)
+				.OnGenerateWidget(this, &SNewSessionRow::OnGenerateServersComboOptionWidget)
+				.ToolTipText(LOCTEXT("SelectServerTooltip", "Select the server on which the session should be created"))
+				[
+					MakeSelectedServerWidget()
+				]
+			]
+
+			+SHorizontalBox::Slot()
+			.AutoWidth()
+			.Padding(2.0f)
+			.HAlign(HAlign_Left)
+			[
+				SNew(SUniformGridPanel)
+				.SlotPadding(FMargin(1.0f, 0.0f))
+
+				// 'Accept' button
+				+SUniformGridPanel::Slot(0, 0)
+				[
+					ConcertBrowserUtils::MakePositiveActionButton(
+						FEditorStyle::GetBrush("Icons.Check"),
+						LOCTEXT("CreateCheckIconTooltip", "Create the session"),
+						TAttribute<bool>::Create([this]() { return !EditableSessionName->GetText().IsEmpty(); }),
+						FOnClicked::CreateRaw(this, &SNewSessionRow::OnAccept))
+				]
+
+				// 'Decline' button
+				+SUniformGridPanel::Slot(1, 0)
+				[
+					ConcertBrowserUtils::MakeNegativeActionButton(
+						FEditorStyle::GetBrush("Icons.X"),
+						LOCTEXT("CancelIconTooltip", "Cancel"),
+						true, // Always enabled.
+						FOnClicked::CreateRaw(this, &SNewSessionRow::OnDecline))
+				]
+			];
+	}
+}
+
+TSharedRef<SWidget> SNewSessionRow::OnGenerateServersComboOptionWidget(TSharedPtr<FConcertServerInfo> ServerItem)
+{
+	bool bIsDefaultServer = ServerItem->ServerName == DefaultServerURL.Get();
+
+	FText Tooltip;
+	if (bIsDefaultServer)
+	{
+		Tooltip = LOCTEXT("DefaultServerTooltip", "Default Configured Server");
+	}
+	else if (ServerItem->ServerName == FPlatformProcess::ComputerName())
+	{
+		Tooltip = LOCTEXT("LocalServerTooltip", "Local Server Running on This Computer");
+	}
+	else
+	{
+		Tooltip = LOCTEXT("OnlineServerTooltip", "Online Server");
+	}
+
+	return SNew(SHorizontalBox)
+		.ToolTipText(Tooltip)
+
+		+SHorizontalBox::Slot()
+		.AutoWidth()
+		[
+			SNew(STextBlock)
+			.Font(bIsDefaultServer ? FAppStyle::Get().GetFontStyle("BoldFont") : FAppStyle::Get().GetFontStyle("NormalFont"))
+			.Text(GetServerDisplayName(ServerItem->ServerName))
+		]
+		+SHorizontalBox::Slot()
+		.AutoWidth()
+		[
+			ConcertBrowserUtils::MakeServerVersionIgnoredWidget(ServerItem->ServerFlags)
+		];
+}
+
+void SNewSessionRow::UpdateServerList()
+{
+	// Remember the currently selected item (if any).
+	TSharedPtr<FConcertServerInfo> SelectedItem = ServersComboBox->GetSelectedItem(); // Instance in current list.
+
+	// Clear the current list. The list is rebuilt from scratch.
+	Servers.Reset();
+
+	TSharedPtr<FConcertServerInfo> LocalServerInfo;
+	TSharedPtr<FConcertServerInfo> DefaultServerInfo;
+	TSharedPtr<FConcertServerInfo> SelectedServerInfo; // Instance in the new list.
+
+	// Convert to shared ptr (slate needs that) and find if the latest list contains a default/local server.
+	for (const FConcertServerInfo& ServerInfo : GetServersFunc())
+	{
+		TSharedPtr<FConcertServerInfo> ComboItem = MakeShared<FConcertServerInfo>(ServerInfo);
+		
+		// Default server is deemed more important than local server to display the icon aside the server.
+		if (ComboItem->ServerName == DefaultServerURL.Get()) 
+		{
+			DefaultServerInfo = ComboItem;
+		}
+		else if (ComboItem->ServerName == FPlatformProcess::ComputerName())
+		{
+			LocalServerInfo = ComboItem;
+		}
+
+		if (SelectedItem.IsValid() && SelectedItem->ServerName == ComboItem->ServerName)
+		{
+			SelectedServerInfo = ComboItem; // Preserve user selection using the new instance.
+		}
+
+		Servers.Emplace(MoveTemp(ComboItem));
+	}
+
+	// Sort the server list alphabetically.
+	Servers.Sort([](const TSharedPtr<FConcertServerInfo>& Lhs, const TSharedPtr<FConcertServerInfo>& Rhs) { return Lhs->ServerName < Rhs->ServerName; });
+
+	// If a server is running on this machine, put it first in the list.
+	if (LocalServerInfo.IsValid() && Servers[0] != LocalServerInfo)
+	{
+		Servers.Remove(LocalServerInfo); // Keep sort order.
+		Servers.Insert(LocalServerInfo, 0);
+	}
+
+	// If a 'default server' is configured and available, put it first in the list. (Possibly overruling the local one)
+	if (DefaultServerInfo.IsValid() && Servers[0] != DefaultServerInfo)
+	{
+		Servers.Remove(DefaultServerInfo); // Keep sort order.
+		Servers.Insert(DefaultServerInfo, 0);
+	}
+
+	// If a server was selected and is still in the updated list.
+	if (SelectedServerInfo.IsValid())
+	{
+		// Preserve user selection.
+		ServersComboBox->SetSelectedItem(SelectedServerInfo);
+	}
+	else if (Servers.Num())
+	{
+		// Select the very first item in the list which is most likely be the default or the local server as they were put first above.
+		ServersComboBox->SetSelectedItem(Servers[0]);
+	}
+	else // Server list is empty.
+	{
+		ServersComboBox->ClearSelection();
+		Servers.Reset();
+	}
+
+	ServersComboBox->RefreshOptions();
+}
+
+TSharedRef<SWidget> SNewSessionRow::MakeSelectedServerWidget()
+{
+	return SNew(SHorizontalBox)
+		+SHorizontalBox::Slot()
+		.AutoWidth()
+		[
+			SNew(STextBlock)
+			.Text(this, &SNewSessionRow::GetSelectedServerText)
+			.HighlightText(HighlightText)
+		]
+		+SHorizontalBox::Slot()
+		.AutoWidth()
+		.Padding(2, 0, 0, 0)
+		[
+			SNew(STextBlock)
+			.Font(FAppStyle::Get().GetFontStyle("FontAwesome.9"))
+			.Text(this, &SNewSessionRow::GetSelectedServerIgnoreVersionText)
+			.ToolTipText(this, &SNewSessionRow::GetSelectedServerIgnoreVersionTooltip)
+		];
+}
+
+FText SNewSessionRow::GetSelectedServerText() const
+{
+	TSharedPtr<FConcertServerInfo> SelectedServer = ServersComboBox->GetSelectedItem();
+	if (SelectedServer.IsValid())
+	{
+		return GetServerDisplayName(SelectedServer->ServerName);
+	}
+	return LOCTEXT("SelectAServer", "Select a Server");
+}
+
+FText SNewSessionRow::GetServerDisplayName(const FString& ServerName) const
+{
+	const bool bIsDefaultServer = ServerName == DefaultServerURL.Get();
+	if (bIsDefaultServer)
+	{
+		return FText::Format(LOCTEXT("DefaultServer", "{0} (Default)"), FText::FromString(FPlatformProcess::ComputerName()));
+	}
+	if (ServerName == FPlatformProcess::ComputerName())
+	{
+		return FText::Format(LOCTEXT("MyComputer", "{0} (My Computer)"), FText::FromString(FPlatformProcess::ComputerName()));
+	}
+	return FText::FromString(ServerName);
+}
+
+FText SNewSessionRow::GetSelectedServerIgnoreVersionText() const
+{
+	if (ServersComboBox->GetSelectedItem() && (ServersComboBox->GetSelectedItem()->ServerFlags & EConcertServerFlags::IgnoreSessionRequirement) != EConcertServerFlags::None)
+	{
+		return FEditorFontGlyphs::Exclamation_Triangle;
+	}
+	return FText();
+}
+
+FText SNewSessionRow::GetSelectedServerIgnoreVersionTooltip() const
+{
+	if (ServersComboBox->GetSelectedItem() && (ServersComboBox->GetSelectedItem()->ServerFlags & EConcertServerFlags::IgnoreSessionRequirement) != EConcertServerFlags::None)
+	{
+		return ConcertBrowserUtils::GetServerVersionIgnoredTooltip();
+	}
+	return FText();
+}
+
+FReply SNewSessionRow::OnAccept()
+{
+	TSharedPtr<FConcertSessionItem> ItemPin = Item.Pin();
+	FString NewSessionName = EditableSessionName->GetText().ToString();
+
+	FText InvalidNameErrorMsg = ConcertSettingsUtils::ValidateSessionName(NewSessionName);
+	if (InvalidNameErrorMsg.IsEmpty())
+	{
+		ItemPin->SessionName = EditableSessionName->GetText().ToString();
+		ItemPin->ServerName = ServersComboBox->GetSelectedItem()->ServerName;
+		ItemPin->ServerAdminEndpointId = ServersComboBox->GetSelectedItem()->AdminEndpointId;
+		AcceptFunc(ItemPin); // Delegate to create the session.
+	}
+	else
+	{
+		EditableSessionName->SetError(InvalidNameErrorMsg);
+		FSlateApplication::Get().SetKeyboardFocus(EditableSessionName);
+	}
+
+	return FReply::Handled();
+}
+
+FReply SNewSessionRow::OnDecline()
+{
+	DeclineFunc(Item.Pin()); // Decline the creation and remove the row from the table.
+	return FReply::Handled();
+}
+
+void SNewSessionRow::OnSessionNameChanged(const FText& NewName)
+{
+	EditableSessionName->SetError(ConcertSettingsUtils::ValidateSessionName(NewName.ToString()));
+}
+
+void SNewSessionRow::OnSessionNameCommitted(const FText& NewText, ETextCommit::Type CommitType)
+{
+	if (CommitType == ETextCommit::OnEnter)
+	{
+		OnAccept(); // Create the session.
+	}
+}
+
+FReply SNewSessionRow::OnKeyDownHandler(const FGeometry&, const FKeyEvent& KeyEvent)
+{
+	// NOTE: This is invoked when the editable text field has the focus.
+	return KeyEvent.GetKey() == EKeys::Escape ? OnDecline() : FReply::Unhandled();
+}
+
+#undef LOCTEXT_NAMESPACE
\ No newline at end of file
diff --git a/Engine/Plugins/Developer/Concert/ConcertUI/ConcertSharedSlate/Source/ConcertSharedSlate/Private/Session/Browser/SNewSessionRow.h b/Engine/Plugins/Developer/Concert/ConcertUI/ConcertSharedSlate/Source/ConcertSharedSlate/Private/Session/Browser/SNewSessionRow.h
new file mode 100644
--- /dev/null
+++ b/Engine/Plugins/Developer/Concert/ConcertUI/ConcertSharedSlate/Source/ConcertSharedSlate/Private/Session/Browser/SNewSessionRow.h
@@ -0,0 +1,81 @@
+﻿// Copyright Epic Games, Inc. All Rights Reserved.
+
+#pragma once
+
+#include "CoreMinimal.h"
+#include "ConcertMessageData.h"
+#include "Widgets/DeclarativeSyntaxSupport.h"
+#include "Widgets/Views/STableRow.h"
+#include "Widgets/Input/SComboBox.h"
+
+class FConcertSessionItem;
+class SEditableTextBox;
+class SWidget;
+
+/**
+ * The type of row used in the session list view to edit a new session (the session name + server).
+ */
+class SNewSessionRow : public SMultiColumnTableRow<TSharedPtr<FConcertSessionItem>>
+{
+public:
+	using FGetServersFunc = TFunction<TArray<FConcertServerInfo>()>;
+	// Should remove the editable 'new' row and creates the sessions.
+	using FAcceptFunc = TFunction<void(const TSharedPtr<FConcertSessionItem>&)>;
+	// Should just remove the editable 'new' row.
+	using FDeclineFunc = TFunction<void(const TSharedPtr<FConcertSessionItem>&)>; 
+
+	SLATE_BEGIN_ARGS(SNewSessionRow) // Needed to use Args because Construct() is limited to 5 arguments and 6 were required.
+		: _GetServerFunc()
+		, _OnAcceptFunc()
+		, _OnDeclineFunc()
+		, _HighlightText()
+	{
+	}
+
+	SLATE_ARGUMENT(FGetServersFunc, GetServerFunc)
+	SLATE_ARGUMENT(FAcceptFunc, OnAcceptFunc)
+	SLATE_ARGUMENT(FDeclineFunc, OnDeclineFunc)
+	
+	SLATE_ATTRIBUTE(FText, HighlightText)
+	SLATE_ATTRIBUTE(FString, DefaultServerURL)
+
+	SLATE_END_ARGS()
+
+	void Construct(const FArguments& InArgs, TSharedPtr<FConcertSessionItem> InItem, const TSharedRef<STableViewBase>& InOwnerTableView);
+	
+	virtual TSharedRef<SWidget> GenerateWidgetForColumn(const FName& ColumnName) override;
+	virtual void Tick(const FGeometry& AllottedGeometry, const double InCurrentTime, const float InDeltaTime) override;
+
+private:
+	
+	TSharedRef<SWidget> OnGenerateServersComboOptionWidget(TSharedPtr<FConcertServerInfo> Item);
+	TSharedRef<SWidget> MakeSelectedServerWidget();
+	FText GetSelectedServerText() const;
+	FText GetSelectedServerIgnoreVersionText() const;
+	FText GetSelectedServerIgnoreVersionTooltip() const;
+	FText GetServerDisplayName(const FString& ServerName) const;
+
+	void OnSessionNameChanged(const FText& NewName);
+	void OnSessionNameCommitted(const FText& NewText, ETextCommit::Type CommitType);
+	FReply OnAccept();
+	FReply OnDecline();
+	FReply OnKeyDownHandler(const FGeometry&, const FKeyEvent&); // Registered as handler to the editable text (vs. OnKeyDown() virtual method).
+
+	void UpdateServerList();
+
+	/** Holds the new item to fill with session name and server. */
+	TWeakPtr<FConcertSessionItem> Item;
+	/** Servers displayed in the server combo box. */
+	TArray<TSharedPtr<FConcertServerInfo>> Servers;
+	
+	TSharedPtr<SComboBox<TSharedPtr<FConcertServerInfo>>> ServersComboBox;
+	TSharedPtr<SEditableTextBox> EditableSessionName;
+	
+	FGetServersFunc GetServersFunc;
+	FAcceptFunc AcceptFunc;
+	FDeclineFunc DeclineFunc;
+	
+	TAttribute<FText> HighlightText;
+	TAttribute<FString> DefaultServerURL;
+	bool bInitialFocusTaken = false;
+};
diff --git a/Engine/Plugins/Developer/Concert/ConcertUI/ConcertSharedSlate/Source/ConcertSharedSlate/Private/Session/Browser/SSaveRestoreSessionRow.cpp b/Engine/Plugins/Developer/Concert/ConcertUI/ConcertSharedSlate/Source/ConcertSharedSlate/Private/Session/Browser/SSaveRestoreSessionRow.cpp
new file mode 100644
--- /dev/null
+++ b/Engine/Plugins/Developer/Concert/ConcertUI/ConcertSharedSlate/Source/ConcertSharedSlate/Private/Session/Browser/SSaveRestoreSessionRow.cpp
@@ -0,0 +1,209 @@
+﻿// Copyright Epic Games, Inc. All Rights Reserved.
+
+#include "SSaveRestoreSessionRow.h"
+
+#include "Session/Browser/ConcertBrowserUtils.h"
+#include "Session/Browser/ConcertSessionItem.h"
+
+#include "EditorStyleSet.h"
+#include "Framework/Application/SlateApplication.h"
+#include "Internationalization/Regex.h"
+#include "Styling/AppStyle.h"
+#include "Widgets/Input/SEditableTextBox.h"
+#include "Widgets/Layout/SBox.h"
+#include "Widgets/Layout/SUniformGridPanel.h"
+#include "Widgets/SBoxPanel.h"
+#include "Widgets/Text/STextBlock.h"
+#include "Widgets/Views/SExpanderArrow.h"
+
+#define LOCTEXT_NAMESPACE "SConcertBrowser"
+
+void SSaveRestoreSessionRow::Construct(const FArguments& InArgs, TSharedPtr<FConcertSessionItem> InNode, const TSharedRef<STableViewBase>& InOwnerTableView)
+{
+	Item = MoveTemp(InNode);
+	AcceptFunc = InArgs._OnAcceptFunc;
+	DeclineFunc = InArgs._OnDeclineFunc;
+	HighlightText = InArgs._HighlightText;
+
+	AcceptFunc.CheckCallable();
+	DeclineFunc.CheckCallable();
+
+	// Construct base class
+	SMultiColumnTableRow<TSharedPtr<FConcertSessionItem>>::Construct(FSuperRowType::FArguments(), InOwnerTableView);
+}
+
+TBitArray<> SSaveRestoreSessionRow::GetWiresNeededByDepth() const
+{
+	TBitArray<> Bits;
+	Bits.Add(false);
+	return Bits;
+}
+
+FText SSaveRestoreSessionRow::GetDefaultName(const FConcertSessionItem& InItem) const
+{
+	if (InItem.Type == FConcertSessionItem::EType::SaveSession)
+	{
+		return FText::Format(LOCTEXT("DefaultName", "{0}.{1}"), FText::FromString(InItem.SessionName), FText::FromString(FDateTime::UtcNow().ToString()));
+	}
+
+	// Supposing the name of the archive has the dates as suffix, like SessionXYZ.2019.03.13-19.39.12, then extracts SessionXYZ
+	static FRegexPattern Pattern(TEXT(R"((.*)\.\d+\.\d+\.\d+\-\d+\.\d+\.\d+$)"));
+	FRegexMatcher Matcher(Pattern, InItem.SessionName);
+	if (Matcher.FindNext())
+	{
+		return FText::FromString(Matcher.GetCaptureGroup(1));
+	}
+
+	return FText::FromString(InItem.SessionName);
+}
+
+TSharedRef<SWidget> SSaveRestoreSessionRow::GenerateWidgetForColumn(const FName& ColumnName)
+{
+	TSharedPtr<FConcertSessionItem> ItemPin = Item.Pin();
+
+	if (ColumnName == ConcertBrowserUtils::IconColName)
+	{
+		return SNew(SHorizontalBox)
+			+SHorizontalBox::Slot()
+			.AutoWidth()
+			.Padding(8, 0)
+			[
+				SNew(SExpanderArrow, SharedThis(this))
+				.StyleSet(&FAppStyle::Get())
+				.ShouldDrawWires(true)
+			];
+	}
+	else if (ColumnName == ConcertBrowserUtils::SessionColName)
+	{
+		return SNew(SHorizontalBox)
+
+			// 'Restore as/Save as' text
+			+SHorizontalBox::Slot()
+			.VAlign(VAlign_Center)
+			.AutoWidth()
+			.Padding(2.0, 0.0)
+			[
+				SNew(STextBlock).Text(ItemPin->Type == FConcertSessionItem::EType::RestoreSession ? LOCTEXT("RestoreAs", "Restore as:") : LOCTEXT("ArchiveAs", "Archive as:"))
+			]
+
+			// Editable text.
+			+SHorizontalBox::Slot()
+			.VAlign(VAlign_Center)
+			.Padding(FMargin(0, 0, 2, 0))
+			[
+				SAssignNew(EditableSessionName, SEditableTextBox)
+				.HintText(ItemPin->Type == FConcertSessionItem::EType::RestoreSession ? LOCTEXT("RestoreSessionHint", "Enter a session name") : LOCTEXT("ArchivSessionHint", "Enter an archive name"))
+				.OnTextCommitted(this, &SSaveRestoreSessionRow::OnSessionNameCommitted)
+				.OnKeyDownHandler(this, &SSaveRestoreSessionRow::OnKeyDownHandler)
+				.OnTextChanged(this, &SSaveRestoreSessionRow::OnSessionNameChanged)
+				.Text(GetDefaultName(*ItemPin))
+				.SelectAllTextWhenFocused(true)
+			];
+	}
+	else
+	{
+		check(ColumnName == ConcertBrowserUtils::ServerColName);
+
+		// 'Server' text block.
+		return SNew(SHorizontalBox)
+		+SHorizontalBox::Slot()
+		.AutoWidth()
+		[
+			SNew(SBox)
+			.VAlign(VAlign_Center)
+			[
+				// Server.
+				SNew(STextBlock)
+				.Text(FText::FromString(ItemPin->ServerName))
+				.HighlightText(HighlightText)
+			]
+		]
+
+		+SHorizontalBox::Slot()
+		.AutoWidth()
+		.Padding(2.0f)
+		.HAlign(HAlign_Left)
+		[
+			SNew(SUniformGridPanel)
+			.SlotPadding(FMargin(1.0f, 0.0f))
+
+			// 'Accept' button
+			+SUniformGridPanel::Slot(0, 0)
+			[
+				ConcertBrowserUtils::MakePositiveActionButton(
+					FEditorStyle::GetBrush("Icons.Check"),
+					ItemPin->Type == FConcertSessionItem::EType::RestoreSession ? LOCTEXT("RestoreCheckIconTooltip", "Restore the session") : LOCTEXT("ArchiveCheckIconTooltip", "Archive the session"),
+					TAttribute<bool>::Create([this]() { return !EditableSessionName->GetText().IsEmpty(); }), // Enabled?
+					FOnClicked::CreateRaw(this, &SSaveRestoreSessionRow::OnAccept))
+			]
+
+			// 'Cancel' button
+			+SUniformGridPanel::Slot(1, 0)
+			[
+				ConcertBrowserUtils::MakeNegativeActionButton(
+					FEditorStyle::GetBrush("Icons.X"),
+					LOCTEXT("CancelTooltip", "Cancel"),
+					true, // Enabled?
+					FOnClicked::CreateRaw(this, &SSaveRestoreSessionRow::OnDecline))
+			]
+		];
+	}
+}
+
+void SSaveRestoreSessionRow::Tick(const FGeometry& AllottedGeometry, const double InCurrentTime, const float InDeltaTime)
+{
+	// Should give the focus to an editable text.
+	if (!bInitialFocusTaken)
+	{
+		bInitialFocusTaken = FSlateApplication::Get().SetKeyboardFocus(EditableSessionName.ToSharedRef());
+	}
+}
+
+void SSaveRestoreSessionRow::OnSessionNameChanged(const FText& NewName)
+{
+	EditableSessionName->SetError(ConcertSettingsUtils::ValidateSessionName(NewName.ToString()));
+}
+
+void SSaveRestoreSessionRow::OnSessionNameCommitted(const FText& NewText, ETextCommit::Type CommitType)
+{
+	TSharedPtr<FConcertSessionItem> ItemPin = Item.Pin();
+	if (CommitType == ETextCommit::Type::OnEnter)
+	{
+		OnAccept();
+	}
+}
+
+FReply SSaveRestoreSessionRow::OnAccept()
+{
+	// Read the session name given by the user.
+	TSharedPtr<FConcertSessionItem> ItemPin = Item.Pin();
+	FString Name = EditableSessionName->GetText().ToString(); // Archive name or restored session name.
+
+	// Ensure the user provided a name.
+	FText InvalidNameErrorMsg = ConcertSettingsUtils::ValidateSessionName(Name);
+	if (InvalidNameErrorMsg.IsEmpty()) // Name is valid?
+	{
+		AcceptFunc(ItemPin, Name); // Delegate archiving/restoring operation.
+	}
+	else
+	{
+		EditableSessionName->SetError(InvalidNameErrorMsg);
+		FSlateApplication::Get().SetKeyboardFocus(EditableSessionName);
+	}
+
+	return FReply::Handled();
+}
+
+FReply SSaveRestoreSessionRow::OnDecline()
+{
+	DeclineFunc(Item.Pin()); // Remove the save/restore editable row.
+	return FReply::Handled();
+}
+
+FReply SSaveRestoreSessionRow::OnKeyDownHandler(const FGeometry&, const FKeyEvent& KeyEvent)
+{
+	// NOTE: This handler is to capture the 'Escape' while the text field has focus.
+	return KeyEvent.GetKey() == EKeys::Escape ? OnDecline() : FReply::Unhandled();
+}
+
+#undef LOCTEXT_NAMESPACE
\ No newline at end of file
diff --git a/Engine/Plugins/Developer/Concert/ConcertUI/ConcertSharedSlate/Source/ConcertSharedSlate/Private/Session/Browser/SSaveRestoreSessionRow.h b/Engine/Plugins/Developer/Concert/ConcertUI/ConcertSharedSlate/Source/ConcertSharedSlate/Private/Session/Browser/SSaveRestoreSessionRow.h
new file mode 100644
--- /dev/null
+++ b/Engine/Plugins/Developer/Concert/ConcertUI/ConcertSharedSlate/Source/ConcertSharedSlate/Private/Session/Browser/SSaveRestoreSessionRow.h
@@ -0,0 +1,67 @@
+﻿// Copyright Epic Games, Inc. All Rights Reserved.
+
+#pragma once
+
+#include "CoreMinimal.h"
+#include "Misc/Attribute.h"
+#include "Session/Browser/ConcertSessionItem.h"
+#include "Widgets/Views/SListView.h"
+#include "Widgets/DeclarativeSyntaxSupport.h"
+#include "Widgets/Views/STableRow.h"
+
+class SWidget;
+class SEditableTextBox;
+
+/**
+ * The type of row used in the session list view to create and archive or restore a session (edit the information required for the operation).
+ */
+class SSaveRestoreSessionRow : public SMultiColumnTableRow<TSharedPtr<FConcertSessionItem>>
+{
+public:
+	/** Should remove the editable row and save or restore the session. */
+	using FAcceptFunc = TFunction<void(TSharedPtr<FConcertSessionItem>, const FString&)>;
+	/** Should only remove the editable row from the table. */
+	using FDeclineFunc = TFunction<void(TSharedPtr<FConcertSessionItem>)>;
+
+	SLATE_BEGIN_ARGS(SSaveRestoreSessionRow)
+		: _OnAcceptFunc()
+		, _OnDeclineFunc()
+		, _HighlightText()
+	{}
+
+	SLATE_ARGUMENT(FAcceptFunc, OnAcceptFunc)
+	SLATE_ARGUMENT(FDeclineFunc, OnDeclineFunc)
+	SLATE_ATTRIBUTE(FText, HighlightText)
+
+	SLATE_END_ARGS()
+	
+	void Construct(const FArguments& InArgs, TSharedPtr<FConcertSessionItem> Node, const TSharedRef<STableViewBase>& InOwnerTableView);
+	virtual TSharedRef<SWidget> GenerateWidgetForColumn(const FName& ColumnName) override;
+	virtual void Tick(const FGeometry& AllottedGeometry, const double InCurrentTime, const float InDeltaTime) override;
+
+	void OnSessionNameChanged(const FText& NewName);
+	void OnSessionNameCommitted(const FText& NewText, ETextCommit::Type CommitType);
+	FReply OnAccept();
+	FReply OnDecline();
+	FReply OnKeyDownHandler(const FGeometry&, const FKeyEvent&); // Registered as handler to the editable text (vs. OnKeyDown() virtual method for this widget).
+
+	// Override SMultiColumnTableRow functions to ensure a wire is drawn between the item to restore and the editable row to link them together.
+	virtual TBitArray<> GetWiresNeededByDepth() const override;
+	virtual bool IsLastChild() const override { return true; }
+	virtual int32 DoesItemHaveChildren() const override { return 0; }
+	virtual bool IsItemExpanded() const override { return false; }
+
+	/** Generates a default name for an archive or a restored session. */
+	FText GetDefaultName(const FConcertSessionItem& Item) const; 
+
+private:
+	
+	TWeakPtr<FConcertSessionItem> Item;
+	TSharedPtr<SEditableTextBox> EditableSessionName;
+	
+	FAcceptFunc AcceptFunc;
+	FDeclineFunc DeclineFunc;
+	
+	TAttribute<FText> HighlightText;
+	bool bInitialFocusTaken = false;
+};
diff --git a/Engine/Plugins/Developer/Concert/ConcertUI/ConcertSharedSlate/Source/ConcertSharedSlate/Private/Session/Browser/SSessionRow.cpp b/Engine/Plugins/Developer/Concert/ConcertUI/ConcertSharedSlate/Source/ConcertSharedSlate/Private/Session/Browser/SSessionRow.cpp
new file mode 100644
--- /dev/null
+++ b/Engine/Plugins/Developer/Concert/ConcertUI/ConcertSharedSlate/Source/ConcertSharedSlate/Private/Session/Browser/SSessionRow.cpp
@@ -0,0 +1,189 @@
+﻿// Copyright Epic Games, Inc. All Rights Reserved.
+
+#include "SSessionRow.h"
+
+#include "Session/Browser/ConcertBrowserUtils.h"
+#include "Session/Browser/ConcertSessionItem.h"
+
+#include "EditorFontGlyphs.h"
+#include "Framework/Application/SlateApplication.h"
+#include "Styling/AppStyle.h"
+#include "Styling/CoreStyle.h"
+#include "Widgets/Layout/SBox.h"
+#include "Widgets/Layout/SSpacer.h"
+#include "Widgets/SBoxPanel.h"
+#include "Widgets/Text/SInlineEditableTextBlock.h"
+#include "Widgets/Text/STextBlock.h"
+
+#define LOCTEXT_NAMESPACE "SConcertBrowser"
+
+void SSessionRow::Construct(const FArguments& InArgs, TSharedPtr<FConcertSessionItem> InItem, const TSharedRef<STableViewBase>& InOwnerTableView)
+{
+	Item = MoveTemp(InItem);
+	DoubleClickFunc = InArgs._OnDoubleClickFunc; // This function should join a session or add a row to restore an archive.
+	RenameFunc = InArgs._OnRenameFunc; // Function invoked to send a rename request to the server.
+	IsDefaultSession = InArgs._IsDefaultSession;
+	HighlightText = InArgs._HighlightText;
+	IsSelected = InArgs._IsSelected;
+
+	DoubleClickFunc.CheckCallable();
+	RenameFunc.CheckCallable();
+	IsDefaultSession.CheckCallable();
+
+	// Construct base class
+	SMultiColumnTableRow<TSharedPtr<FConcertSessionItem>>::Construct(FSuperRowType::FArguments(), InOwnerTableView);
+
+	// Listen and handle rename request.
+	InItem->OnBeginEditSessionNameRequest.AddSP(this, &SSessionRow::OnBeginEditingSessionName);
+}
+
+TSharedRef<SWidget> SSessionRow::GenerateWidgetForColumn(const FName& ColumnName)
+{
+	TSharedPtr<FConcertSessionItem> ItemPin = Item.Pin();
+
+	if (ColumnName == ConcertBrowserUtils::IconColName)
+	{
+		return SNew(SBox)
+		.VAlign(VAlign_Center)
+		.HAlign(HAlign_Center)
+		.Padding(2)
+		.ToolTipText(ItemPin->Type == FConcertSessionItem::EType::ActiveSession ? LOCTEXT("ActiveIconTooltip", "Active session") : LOCTEXT("ArchivedIconTooltip", "Archived Session"))
+		[
+			SNew(STextBlock)
+			.Font(FAppStyle::Get().GetFontStyle(ConcertBrowserUtils::IconColumnFontName))
+			.Text(ItemPin->Type == FConcertSessionItem::EType::ActiveSession ? FEditorFontGlyphs::Circle : FEditorFontGlyphs::Archive)
+			.ColorAndOpacity(ItemPin->Type == FConcertSessionItem::EType::ActiveSession ? FAppStyle::Get().GetWidgetStyle<FButtonStyle>("FlatButton.Success").Normal.TintColor : FSlateColor::UseSubduedForeground())
+		];
+	}
+	
+	const bool bIsDefaultConfig = IsDefaultSession(ItemPin);
+	FSlateFontInfo FontInfo;
+	FSlateColor FontColor;
+	if (ItemPin->Type == FConcertSessionItem::EType::ActiveSession)
+	{
+		FontColor = bIsDefaultConfig ? FSlateColor(FLinearColor::White) : FSlateColor(FLinearColor::White * 0.8f);
+		FontInfo = FAppStyle::Get().GetFontStyle("NormalFont");
+	}
+	else
+	{
+		FontColor = FSlateColor::UseSubduedForeground();
+		FontInfo = FCoreStyle::GetDefaultFontStyle("Italic", 9);
+	}
+
+	if (ColumnName == ConcertBrowserUtils::SessionColName)
+	{
+		return SNew(SBox)
+			.VAlign(VAlign_Center)
+			[
+				SAssignNew(SessionNameText, SInlineEditableTextBlock)
+				.Text_Lambda([this]() { return FText::AsCultureInvariant(Item.Pin()->SessionName); })
+				.HighlightText(HighlightText)
+				.OnTextCommitted(this, &SSessionRow::OnSessionNameCommitted)
+				.IsReadOnly(false)
+				.IsSelected(FIsSelected::CreateLambda([this]() { return IsSelected.Get(); }))
+				.OnVerifyTextChanged(this, &SSessionRow::OnValidatingSessionName)
+				.Font(FontInfo)
+				.ColorAndOpacity(FontColor)
+			];
+	}
+
+	check(ColumnName == ConcertBrowserUtils::ServerColName);
+
+	if (bIsDefaultConfig)
+	{
+		return SNew(SHorizontalBox)
+			+SHorizontalBox::Slot()
+			.AutoWidth()
+			.VAlign(VAlign_Center)
+			[
+				SNew(STextBlock)
+				.Text(FText::Format(INVTEXT("{0} * "), FText::AsCultureInvariant(ItemPin->ServerName)))
+				.HighlightText(HighlightText)
+				.Font(FontInfo)
+				.ColorAndOpacity(FontColor)
+			]
+			+SHorizontalBox::Slot()
+			.AutoWidth()
+			.VAlign(VAlign_Center)
+			[
+				SNew(STextBlock)
+				.Text(LOCTEXT("DefaultServerSession", "(Default Session/Server)"))
+				.HighlightText(HighlightText)
+				.Font(FCoreStyle::GetDefaultFontStyle("Regular", 9))
+				.ColorAndOpacity(FontColor)
+			]
+			+SHorizontalBox::Slot()
+			[
+				SNew(SSpacer)
+			]
+			+SHorizontalBox::Slot()
+			.AutoWidth()
+			.VAlign(VAlign_Center)
+			.HAlign(HAlign_Right)
+			[
+				ConcertBrowserUtils::MakeServerVersionIgnoredWidget(ItemPin->ServerFlags)
+			];
+	}
+	else
+	{
+		return SNew(SHorizontalBox)
+			+SHorizontalBox::Slot()
+			.AutoWidth()
+			.VAlign(VAlign_Center)
+			[
+				SNew(STextBlock)
+				.Text(FText::AsCultureInvariant(ItemPin->ServerName))
+				.HighlightText(HighlightText)
+				.Font(FontInfo)
+				.ColorAndOpacity(FontColor)
+			]
+			+SHorizontalBox::Slot()
+			[
+				SNew(SSpacer)
+			]
+			+SHorizontalBox::Slot()
+			.AutoWidth()
+			.VAlign(VAlign_Center)
+			.HAlign(HAlign_Right)
+			[
+				ConcertBrowserUtils::MakeServerVersionIgnoredWidget(ItemPin->ServerFlags)
+			];
+	}
+}
+
+FReply SSessionRow::OnMouseButtonDoubleClick(const FGeometry& InMyGeometry, const FPointerEvent& InMouseEvent)
+{
+	if (TSharedPtr<FConcertSessionItem> ItemPin = Item.Pin())
+	{
+		DoubleClickFunc(ItemPin);
+	}
+	return FReply::Handled();
+}
+
+bool SSessionRow::OnValidatingSessionName(const FText& NewSessionName, FText& OutError)
+{
+	OutError = ConcertSettingsUtils::ValidateSessionName(NewSessionName.ToString());
+	return OutError.IsEmpty();
+}
+
+void SSessionRow::OnSessionNameCommitted(const FText& NewSessionName, ETextCommit::Type CommitType)
+{
+	if (TSharedPtr<FConcertSessionItem> ItemPin = Item.Pin())
+	{
+		FString NewName = NewSessionName.ToString();
+		if (NewName != ItemPin->SessionName) // Was renamed?
+		{
+			if (ConcertSettingsUtils::ValidateSessionName(NewName).IsEmpty()) // Name is valid?
+			{
+				RenameFunc(ItemPin, NewName); // Send the rename request to the server. (Server may still refuse at this point)
+			}
+			else
+			{
+				// NOTE: Error are interactively detected and raised by OnValidatingSessionName()
+				FSlateApplication::Get().SetKeyboardFocus(SessionNameText);
+			}
+		}
+	}
+}
+
+#undef LOCTEXT_NAMESPACE
\ No newline at end of file
diff --git a/Engine/Plugins/Developer/Concert/ConcertUI/ConcertSharedSlate/Source/ConcertSharedSlate/Private/Session/Browser/SSessionRow.h b/Engine/Plugins/Developer/Concert/ConcertUI/ConcertSharedSlate/Source/ConcertSharedSlate/Private/Session/Browser/SSessionRow.h
new file mode 100644
--- /dev/null
+++ b/Engine/Plugins/Developer/Concert/ConcertUI/ConcertSharedSlate/Source/ConcertSharedSlate/Private/Session/Browser/SSessionRow.h
@@ -0,0 +1,63 @@
+﻿// Copyright Epic Games, Inc. All Rights Reserved.
+
+#pragma once
+
+#include "CoreMinimal.h"
+#include "Misc/Attribute.h"
+#include "Widgets/DeclarativeSyntaxSupport.h"
+#include "Widgets/Text/SInlineEditableTextBlock.h"
+#include "Widgets/Views/STableRow.h"
+#include "Widgets/Views/SListView.h"
+
+class FConcertSessionItem;
+class SWidget;
+class SInlineEditableTextBlock;
+
+class SSessionRow : public SMultiColumnTableRow<TSharedPtr<FConcertSessionItem>>
+{
+public:
+	using FDoubleClickFunc = TFunction<void(TSharedPtr<FConcertSessionItem>)>;
+	using FRenameFunc = TFunction<void(TSharedPtr<FConcertSessionItem>, const FString&)>;
+	using FIsDefaultSession = TFunction<bool(TSharedPtr<FConcertSessionItem>)>;
+
+	SLATE_BEGIN_ARGS(SSessionRow)
+		: _OnDoubleClickFunc()
+		, _OnRenameFunc()
+		, _HighlightText()
+		, _IsSelected(false)
+	{}
+
+	SLATE_ARGUMENT(FDoubleClickFunc, OnDoubleClickFunc)
+	SLATE_ARGUMENT(FRenameFunc, OnRenameFunc)
+	SLATE_ARGUMENT(FIsDefaultSession, IsDefaultSession)
+	
+	SLATE_ATTRIBUTE(FText, HighlightText)
+	SLATE_ATTRIBUTE(bool, IsSelected)
+
+	SLATE_END_ARGS()
+
+	void Construct(const FArguments& InArgs, TSharedPtr<FConcertSessionItem> InItem, const TSharedRef<STableViewBase>& InOwnerTableView);
+	
+	virtual TSharedRef<SWidget> GenerateWidgetForColumn(const FName& ColumnName) override;
+	virtual FReply OnMouseButtonDoubleClick( const FGeometry& InMyGeometry, const FPointerEvent& InMouseEvent) override;
+
+	void OnSessionNameCommitted(const FText& NewSessionName, ETextCommit::Type CommitType);
+
+private:
+	
+	void OnBeginEditingSessionName() { SessionNameText->EnterEditingMode(); }
+	bool OnValidatingSessionName(const FText& NewSessionName, FText& OutError);
+
+	TWeakPtr<FConcertSessionItem> Item;
+
+	/** Invoked when the user double click on the row */
+	FDoubleClickFunc DoubleClickFunc;
+	/** Invoked when the user commit the session rename. (This will send the request to server) */
+	FRenameFunc RenameFunc;
+	/** Given a session checks whether it is the default session the client is supposed to join */
+	FIsDefaultSession IsDefaultSession;
+	
+	TAttribute<FText> HighlightText;
+	TAttribute<bool> IsSelected;
+	TSharedPtr<SInlineEditableTextBlock> SessionNameText;
+};
diff --git a/Engine/Plugins/Developer/Concert/ConcertUI/ConcertSharedSlate/Source/ConcertSharedSlate/Private/Session/History/AbstractSessionHistoryController.cpp b/Engine/Plugins/Developer/Concert/ConcertUI/ConcertSharedSlate/Source/ConcertSharedSlate/Private/Session/History/AbstractSessionHistoryController.cpp
new file mode 100644
--- /dev/null
+++ b/Engine/Plugins/Developer/Concert/ConcertUI/ConcertSharedSlate/Source/ConcertSharedSlate/Private/Session/History/AbstractSessionHistoryController.cpp
@@ -0,0 +1,32 @@
+﻿// Copyright Epic Games, Inc. All Rights Reserved.
+
+#include "Session/History/AbstractSessionHistoryController.h"
+
+#include "Session/History/SSessionHistory.h"
+#include "Widgets/DeclarativeSyntaxSupport.h"
+
+FAbstractSessionHistoryController::FAbstractSessionHistoryController(FName PackageFilter)
+	: SessionHistory(MakeSessionHistory(PackageFilter))
+{}
+
+void FAbstractSessionHistoryController::ReloadActivities()
+{
+	constexpr int64 MaximumNumberOfActivities = SSessionHistory::MaximumNumberOfActivities;
+	TMap<FGuid, FConcertClientInfo> EndpointClientInfoMap;
+	TArray<FConcertSessionActivity> FetchedActivities;
+	GetActivities(MaximumNumberOfActivities, EndpointClientInfoMap, FetchedActivities);
+
+	// MoveTemp strictly not needed - but it will be faster on Debug builds
+	GetSessionHistory()->ReloadActivities(
+		MoveTemp(EndpointClientInfoMap),
+		MoveTemp(FetchedActivities)
+		);
+}
+
+TSharedRef<SSessionHistory> FAbstractSessionHistoryController::MakeSessionHistory(FName PackageFilter) const
+{
+	return SNew(SSessionHistory)
+		.PackageFilter(PackageFilter)
+		.GetPackageEvent_Raw(this, &FAbstractSessionHistoryController::GetPackageEvent)
+		.GetTransactionEvent_Raw(this, &FAbstractSessionHistoryController::GetTransactionEvent);
+}
diff --git a/Engine/Plugins/Developer/Concert/ConcertUI/ConcertSharedSlate/Source/ConcertSharedSlate/Private/Session/History/SSessionHistory.cpp b/Engine/Plugins/Developer/Concert/ConcertUI/ConcertSharedSlate/Source/ConcertSharedSlate/Private/Session/History/SSessionHistory.cpp
new file mode 100644
--- /dev/null
+++ b/Engine/Plugins/Developer/Concert/ConcertUI/ConcertSharedSlate/Source/ConcertSharedSlate/Private/Session/History/SSessionHistory.cpp
@@ -0,0 +1,175 @@
+// Copyright Epic Games, Inc. All Rights Reserved.
+
+#include "Session/History/SSessionHistory.h"
+
+#include "ConcertMessageData.h"
+#include "IConcertSession.h"
+#include "Session/Activity/SConcertSessionActivities.h"
+
+#include "Widgets/Input/SSearchBox.h"
+#include "Widgets/Layout/SSeparator.h"
+
+#define LOCTEXT_NAMESPACE "SSessionHistory"
+
+namespace ConcertSessionHistoryUI
+{
+	bool PackageNamePassesFilter(const FName& PackageNameFilter, const TStructOnScope<FConcertSyncActivitySummary>& InActivitySummary)
+	{
+		if (PackageNameFilter.IsNone())
+		{
+			return true;
+		}
+
+		if (const FConcertSyncLockActivitySummary* Summary = InActivitySummary.Cast<FConcertSyncLockActivitySummary>())
+		{
+			return Summary->PrimaryPackageName == PackageNameFilter;
+		}
+
+		if (const FConcertSyncTransactionActivitySummary* Summary = InActivitySummary.Cast<FConcertSyncTransactionActivitySummary>())
+		{
+			return Summary->PrimaryPackageName == PackageNameFilter;
+		}
+
+		if (const FConcertSyncPackageActivitySummary* Summary = InActivitySummary.Cast<FConcertSyncPackageActivitySummary>())
+		{
+			return Summary->PackageName == PackageNameFilter;
+		}
+
+		return false;
+	}
+}
+
+void SSessionHistory::Construct(const FArguments& InArgs)
+{
+	PackageNameFilter = InArgs._PackageFilter;
+
+	ActivityMap.Reserve(MaximumNumberOfActivities);
+	ActivityListViewOptions = MakeShared<FConcertSessionActivitiesOptions>();
+
+	SAssignNew(ActivityListView, SConcertSessionActivities)
+		.OnGetPackageEvent(InArgs._GetPackageEvent)
+		.OnGetTransactionEvent(InArgs._GetTransactionEvent)
+		.OnMapActivityToClient(this, &SSessionHistory::GetClientInfo)
+		.HighlightText(this, &SSessionHistory::HighlightSearchedText)
+		.TimeFormat(ActivityListViewOptions.Get(), &FConcertSessionActivitiesOptions::GetTimeFormat)
+		.ClientNameColumnVisibility(EVisibility::Visible)
+		.ClientAvatarColorColumnVisibility(EVisibility::Visible)
+		.OperationColumnVisibility(EVisibility::Visible)
+		.PackageColumnVisibility(EVisibility::Collapsed)
+		.ConnectionActivitiesVisibility(ActivityListViewOptions.Get(), &FConcertSessionActivitiesOptions::GetConnectionActivitiesVisibility)
+		.LockActivitiesVisibility(ActivityListViewOptions.Get(), &FConcertSessionActivitiesOptions::GetLockActivitiesVisibility)
+		.PackageActivitiesVisibility(ActivityListViewOptions.Get(), &FConcertSessionActivitiesOptions::GetPackageActivitiesVisibility)
+		.TransactionActivitiesVisibility(ActivityListViewOptions.Get(), &FConcertSessionActivitiesOptions::GetTransactionActivitiesVisibility)
+		.DetailsAreaVisibility(EVisibility::Visible)
+		.IsAutoScrollEnabled(true);
+
+	ChildSlot
+	[
+		SNew(SVerticalBox)
+		+SVerticalBox::Slot()
+		.AutoHeight()
+		.Padding(1, 1)
+		[
+			SAssignNew(SearchBox, SSearchBox)
+			.HintText(LOCTEXT("SearchHint", "Search..."))
+			.OnTextChanged(this, &SSessionHistory::OnSearchTextChanged)
+			.OnTextCommitted(this, &SSessionHistory::OnSearchTextCommitted)
+			.DelayChangeNotificationsWhileTyping(true)
+		]
+
+		+SVerticalBox::Slot()
+		[
+			ActivityListView.ToSharedRef()
+		]
+
+		+SVerticalBox::Slot()
+		.AutoHeight()
+		.Padding(4, 3)
+		[
+			SNew(SSeparator)
+		]
+
+		+SVerticalBox::Slot()
+		.AutoHeight()
+		.Padding(4, 0, 4, 3)
+		[
+			ActivityListViewOptions->MakeStatusBar(
+				TAttribute<int32>(ActivityListView.Get(), &SConcertSessionActivities::GetTotalActivityNum),
+				TAttribute<int32>(ActivityListView.Get(), &SConcertSessionActivities::GetDisplayedActivityNum))
+		]
+	];
+}
+
+void SSessionHistory::ReloadActivities(TMap<FGuid, FConcertClientInfo> InEndpointClientInfoMap, TArray<FConcertSessionActivity> InFetchedActivities)
+{
+	EndpointClientInfoMap = MoveTemp(InEndpointClientInfoMap);
+	ActivityMap.Reset();
+	ActivityListView->ResetActivityList(); 
+
+	for (FConcertSessionActivity& FetchedActivity : InFetchedActivities)
+	{
+		if (ConcertSessionHistoryUI::PackageNamePassesFilter(PackageNameFilter, FetchedActivity.ActivitySummary))
+		{
+			TSharedRef<FConcertSessionActivity> NewActivity = MakeShared<FConcertSessionActivity>(MoveTemp(FetchedActivity));
+			ActivityMap.Add(NewActivity->Activity.ActivityId, NewActivity);
+			ActivityListView->Append(NewActivity);
+		}
+	}
+}
+
+void SSessionHistory::HandleActivityAddedOrUpdated(const FConcertClientInfo& InClientInfo, const FConcertSyncActivity& InActivity, const FStructOnScope& InActivitySummary)
+{
+	TStructOnScope<FConcertSyncActivitySummary> ActivitySummary;
+	ActivitySummary.InitializeFromChecked(InActivitySummary);
+
+	if (ConcertSessionHistoryUI::PackageNamePassesFilter(PackageNameFilter, ActivitySummary))
+	{
+		EndpointClientInfoMap.Add(InActivity.EndpointId, InClientInfo);
+
+		if (TSharedPtr<FConcertSessionActivity> ExistingActivity = ActivityMap.FindRef(InActivity.ActivityId))
+		{
+			ExistingActivity->Activity = InActivity;
+			ExistingActivity->ActivitySummary = MoveTemp(ActivitySummary);
+			ActivityListView->RequestRefresh();
+		}
+		else
+		{
+			TSharedRef<FConcertSessionActivity> NewActivity = MakeShared<FConcertSessionActivity>();
+			NewActivity->Activity = InActivity;
+			NewActivity->ActivitySummary = MoveTemp(ActivitySummary);
+
+			ActivityMap.Add(NewActivity->Activity.ActivityId, NewActivity);
+			ActivityListView->Append(NewActivity);
+		}
+	}
+}
+
+void SSessionHistory::OnSearchTextChanged(const FText& InSearchText)
+{
+	SearchedText = InSearchText;
+	SearchBox->SetError(ActivityListView->UpdateTextFilter(InSearchText));
+}
+
+void SSessionHistory::OnSearchTextCommitted(const FText& InSearchText, ETextCommit::Type CommitType)
+{
+	if (!InSearchText.EqualTo(SearchedText))
+	{
+		OnSearchTextChanged(InSearchText);
+	}
+}
+
+FText SSessionHistory::HighlightSearchedText() const
+{
+	return SearchedText;
+}
+
+TOptional<FConcertClientInfo> SSessionHistory::GetClientInfo(FGuid Guid) const
+{
+	if (const FConcertClientInfo* ClientInfo = EndpointClientInfoMap.Find(Guid))
+	{
+		return *ClientInfo;
+	}
+	return {};
+}
+
+#undef LOCTEXT_NAMESPACE /* SSessionHistory */
diff --git a/Engine/Plugins/Developer/Concert/ConcertUI/ConcertSharedSlate/Source/ConcertSharedSlate/Private/Session/History/SSessionHistoryWrapper.cpp b/Engine/Plugins/Developer/Concert/ConcertUI/ConcertSharedSlate/Source/ConcertSharedSlate/Private/Session/History/SSessionHistoryWrapper.cpp
new file mode 100644
--- /dev/null
+++ b/Engine/Plugins/Developer/Concert/ConcertUI/ConcertSharedSlate/Source/ConcertSharedSlate/Private/Session/History/SSessionHistoryWrapper.cpp
@@ -0,0 +1,12 @@
+﻿// Copyright Epic Games, Inc. All Rights Reserved.
+
+#include "Session/History/SSessionHistoryWrapper.h"
+
+void SSessionHistoryWrapper::Construct(const FArguments& InArgs, TSharedRef<FAbstractSessionHistoryController> InController)
+{
+	Controller = InController;
+	ChildSlot
+	[
+		Controller->GetSessionHistory()
+	];
+}
\ No newline at end of file
diff --git a/Engine/Plugins/Developer/Concert/ConcertUI/ConcertSharedSlate/Source/ConcertSharedSlate/Private/SessionBrowser/SConcertSessionBrowser.cpp b/Engine/Plugins/Developer/Concert/ConcertUI/ConcertSharedSlate/Source/ConcertSharedSlate/Private/SessionBrowser/SConcertSessionBrowser.cpp
deleted file mode 100644
diff --git a/Engine/Plugins/Developer/Concert/ConcertUI/ConcertSharedSlate/Source/ConcertSharedSlate/Private/SessionBrowser/SNewSessionRow.cpp b/Engine/Plugins/Developer/Concert/ConcertUI/ConcertSharedSlate/Source/ConcertSharedSlate/Private/SessionBrowser/SNewSessionRow.cpp
deleted file mode 100644
diff --git a/Engine/Plugins/Developer/Concert/ConcertUI/ConcertSharedSlate/Source/ConcertSharedSlate/Private/SessionBrowser/SNewSessionRow.h b/Engine/Plugins/Developer/Concert/ConcertUI/ConcertSharedSlate/Source/ConcertSharedSlate/Private/SessionBrowser/SNewSessionRow.h
deleted file mode 100644
diff --git a/Engine/Plugins/Developer/Concert/ConcertUI/ConcertSharedSlate/Source/ConcertSharedSlate/Private/SessionBrowser/SSaveRestoreSessionRow.cpp b/Engine/Plugins/Developer/Concert/ConcertUI/ConcertSharedSlate/Source/ConcertSharedSlate/Private/SessionBrowser/SSaveRestoreSessionRow.cpp
deleted file mode 100644
diff --git a/Engine/Plugins/Developer/Concert/ConcertUI/ConcertSharedSlate/Source/ConcertSharedSlate/Private/SessionBrowser/SSaveRestoreSessionRow.h b/Engine/Plugins/Developer/Concert/ConcertUI/ConcertSharedSlate/Source/ConcertSharedSlate/Private/SessionBrowser/SSaveRestoreSessionRow.h
deleted file mode 100644
diff --git a/Engine/Plugins/Developer/Concert/ConcertUI/ConcertSharedSlate/Source/ConcertSharedSlate/Private/SessionBrowser/SSessionRow.cpp b/Engine/Plugins/Developer/Concert/ConcertUI/ConcertSharedSlate/Source/ConcertSharedSlate/Private/SessionBrowser/SSessionRow.cpp
deleted file mode 100644
diff --git a/Engine/Plugins/Developer/Concert/ConcertUI/ConcertSharedSlate/Source/ConcertSharedSlate/Private/SessionBrowser/SSessionRow.h b/Engine/Plugins/Developer/Concert/ConcertUI/ConcertSharedSlate/Source/ConcertSharedSlate/Private/SessionBrowser/SSessionRow.h
deleted file mode 100644
diff --git a/Engine/Plugins/Developer/Concert/ConcertUI/ConcertSharedSlate/Source/ConcertSharedSlate/Private/SessionHistory/AbstractSessionHistoryController.cpp b/Engine/Plugins/Developer/Concert/ConcertUI/ConcertSharedSlate/Source/ConcertSharedSlate/Private/SessionHistory/AbstractSessionHistoryController.cpp
deleted file mode 100644
diff --git a/Engine/Plugins/Developer/Concert/ConcertUI/ConcertSharedSlate/Source/ConcertSharedSlate/Private/SessionHistory/SSessionHistory.cpp b/Engine/Plugins/Developer/Concert/ConcertUI/ConcertSharedSlate/Source/ConcertSharedSlate/Private/SessionHistory/SSessionHistory.cpp
deleted file mode 100644
diff --git a/Engine/Plugins/Developer/Concert/ConcertUI/ConcertSharedSlate/Source/ConcertSharedSlate/Private/SessionHistory/SSessionHistoryWrapper.cpp b/Engine/Plugins/Developer/Concert/ConcertUI/ConcertSharedSlate/Source/ConcertSharedSlate/Private/SessionHistory/SSessionHistoryWrapper.cpp
deleted file mode 100644
diff --git a/Engine/Plugins/Developer/Concert/ConcertUI/ConcertSharedSlate/Source/ConcertSharedSlate/Public/SConcertSessionActivities.h b/Engine/Plugins/Developer/Concert/ConcertUI/ConcertSharedSlate/Source/ConcertSharedSlate/Public/SConcertSessionActivities.h
deleted file mode 100644
diff --git a/Engine/Plugins/Developer/Concert/ConcertUI/ConcertSharedSlate/Source/ConcertSharedSlate/Public/SConcertSessionRecovery.h b/Engine/Plugins/Developer/Concert/ConcertUI/ConcertSharedSlate/Source/ConcertSharedSlate/Public/SConcertSessionRecovery.h
--- a/Engine/Plugins/Developer/Concert/ConcertUI/ConcertSharedSlate/Source/ConcertSharedSlate/Public/SConcertSessionRecovery.h
+++ b/Engine/Plugins/Developer/Concert/ConcertUI/ConcertSharedSlate/Source/ConcertSharedSlate/Public/SConcertSessionRecovery.h
@@ -3,14 +3,13 @@
 #pragma once
 
 #include "CoreMinimal.h"
-#include "SConcertSessionActivities.h"
+#include "Session/Activity/SConcertSessionActivities.h"
 #include "SlateFwd.h"
 #include "Widgets/DeclarativeSyntaxSupport.h"
 #include "Widgets/SCompoundWidget.h"
 
 struct FConcertSessionActivity;
 struct FConcertClientInfo;
-class SConcertSessionActivities;
 class FConcertSessionActivitiesOptions;
 
 /**
diff --git a/Engine/Plugins/Developer/Concert/ConcertUI/ConcertSharedSlate/Source/ConcertSharedSlate/Public/Session/Activity/SConcertSessionActivities.h b/Engine/Plugins/Developer/Concert/ConcertUI/ConcertSharedSlate/Source/ConcertSharedSlate/Public/Session/Activity/SConcertSessionActivities.h
new file mode 100644
--- /dev/null
+++ b/Engine/Plugins/Developer/Concert/ConcertUI/ConcertSharedSlate/Source/ConcertSharedSlate/Public/Session/Activity/SConcertSessionActivities.h
@@ -0,0 +1,407 @@
+// Copyright Epic Games, Inc. All Rights Reserved.
+
+#pragma once
+
+#include "CoreMinimal.h"
+#include "SlateFwd.h"
+#include "Async/Future.h"
+#include "Misc/TextFilter.h"
+#include "Widgets/DeclarativeSyntaxSupport.h"
+#include "Widgets/SCompoundWidget.h"
+#include "Widgets/Views/SListView.h"
+#include "ConcertSyncSessionTypes.h"
+
+struct FConcertSessionActivity;
+struct FConcertClientInfo;
+struct FConcertPackageInfo;
+struct FConcertTransactionEventBase;
+class SUndoHistoryDetails;
+class SPackageDetails;
+class STextBlock;
+class SExpendableArea;
+
+
+/** Filters for the concert session activity view. */
+enum class EConcertActivityFilterFlags
+{
+	ShowAll                  = 0x00,
+	HideConnectionActivities = 1<<0,
+	HideLockActivities       = 1<<1,
+	HidePackageActivities    = 1<<2,
+	HideTransactionActivities= 1<<3,
+	HideIgnoredActivities    = 1<<4,
+};
+ENUM_CLASS_FLAGS(EConcertActivityFilterFlags);
+
+/** Displays session activities in a table view. */
+class CONCERTSHAREDSLATE_API SConcertSessionActivities : public SCompoundWidget
+{
+public:
+	
+	/** Defines how the time should be displayed in the date/time column. */
+	enum class ETimeFormat
+	{
+		Relative, // Display relative time (23 seconds ago)
+		Absolute  // Display absolute time (April 7, 2019 - 10:33:52)
+	};
+
+	/** Used to pull activities from a session. Used to fetch and display the activities of an archived session. */
+	DECLARE_DELEGATE_RetVal_ThreeParams(bool, FFetchActivitiesFunc, TArray<TSharedPtr<FConcertSessionActivity>>& /*InOutActivities*/, int32& /*OutFetchedCount*/, FText& /*ErrorMsg*/);
+	
+	/** Used to map an activity to its client. */
+	DECLARE_DELEGATE_RetVal_OneParam(TOptional<FConcertClientInfo>, FGetActivityClientInfoFunc, FGuid /*ClientId*/);
+
+	/** Returns the transaction event corresponding the specified activity.*/
+	DECLARE_DELEGATE_RetVal_OneParam(TFuture<TOptional<FConcertSyncTransactionEvent>>, FGetTransactionEvent, const FConcertSessionActivity& /*Activity*/)
+	
+
+	/** Returns the package event corresponding to the package activity. */
+	DECLARE_DELEGATE_RetVal_TwoParams(bool, FGetPackageEvent, const FConcertSessionActivity& /*Activity*/, FConcertSyncPackageEventMetaData& /*OutEvent*/);
+
+	/** Used to overlay a widget over a column widget to add custom functionalities to a row. */
+	DECLARE_DELEGATE_RetVal_TwoParams(TSharedPtr<SWidget>, FMakeColumnOverlayWidgetFunc, TWeakPtr<FConcertSessionActivity> /*ThisRowActivity*/, const FName& /*ColumnId*/);
+
+public:
+	SLATE_BEGIN_ARGS(SConcertSessionActivities)
+		: _OnFetchActivities()
+		, _OnMapActivityToClient()
+		, _OnMakeColumnOverlayWidget()
+		, _TimeFormat(ETimeFormat::Relative)
+		, _ClientAvatarColorColumnVisibility(EVisibility::Hidden)
+		, _ClientNameColumnVisibility(EVisibility::Hidden)
+		, _OperationColumnVisibility(EVisibility::Hidden)
+		, _PackageColumnVisibility(EVisibility::Hidden)
+		, _ConnectionActivitiesVisibility(EVisibility::Hidden)
+		, _LockActivitiesVisibility(EVisibility::Hidden)
+		, _PackageActivitiesVisibility(EVisibility::Visible)
+		, _TransactionActivitiesVisibility(EVisibility::Visible)
+		, _IgnoredActivitiesVisibility(EVisibility::Hidden)
+		, _DetailsAreaVisibility(EVisibility::Hidden)
+		, _IsAutoScrollEnabled(false){ }
+
+		/** If bound, invoked to populate the view. */
+		SLATE_EVENT(FFetchActivitiesFunc, OnFetchActivities)
+
+		/** If bound, invoked to map an activity to a client.*/
+		SLATE_EVENT(FGetActivityClientInfoFunc, OnMapActivityToClient)
+
+		/** If bound, invoked to fill up the package activity details panel. */
+		SLATE_EVENT(FGetPackageEvent, OnGetPackageEvent)
+
+		/** If bound, invoked to fill up the transaction activity details panel. */
+		SLATE_EVENT(FGetTransactionEvent, OnGetTransactionEvent)
+
+		/** If bound, invoked when generating a row to add an overlay to a column. */
+		SLATE_EVENT(FMakeColumnOverlayWidgetFunc, OnMakeColumnOverlayWidget)
+
+		/** Highlight the returned text in the view. */
+		SLATE_ATTRIBUTE(FText, HighlightText)
+
+		/** Defines how time should be displayed (relative vs absolute). */
+		SLATE_ATTRIBUTE(ETimeFormat, TimeFormat)
+
+		/** Show/hide the column displaying the avatar color of the client who performed the activity. */
+		SLATE_ATTRIBUTE(EVisibility, ClientAvatarColorColumnVisibility)
+
+		/** Show/hide the column showing the display name of the client who performed the activity. */
+		SLATE_ATTRIBUTE(EVisibility, ClientNameColumnVisibility)
+
+		/** Show/hide the column showing the operation name. */
+		SLATE_ATTRIBUTE(EVisibility, OperationColumnVisibility)
+
+		/** Show/hide the column showing the affected package. */
+		SLATE_ATTRIBUTE(EVisibility, PackageColumnVisibility)
+
+		/** Show/hide connection activities. */
+		SLATE_ATTRIBUTE(EVisibility, ConnectionActivitiesVisibility)
+
+		/** Show/hide lock activities. */
+		SLATE_ATTRIBUTE(EVisibility, LockActivitiesVisibility)
+
+		/** Show/hide package activities. */
+		SLATE_ATTRIBUTE(EVisibility, PackageActivitiesVisibility)
+
+		/** Show/hide transaction activities. */
+		SLATE_ATTRIBUTE(EVisibility, TransactionActivitiesVisibility)
+
+		/** Show/hide ignored activities. */
+		SLATE_ATTRIBUTE(EVisibility, IgnoredActivitiesVisibility)
+
+		/** Show/hide the details area widget. (Not to confuse with widget expansion state) */
+		SLATE_ARGUMENT(EVisibility, DetailsAreaVisibility)
+
+		/** True to scroll the list down automatically (unless the user manually scrolled the list). */
+		SLATE_ARGUMENT(bool, IsAutoScrollEnabled)
+
+		/** Show/hide a message overlay above the activities list explaining why no activities are displayed. */
+		SLATE_ATTRIBUTE(FText, NoActivitiesReasonText)
+	SLATE_END_ARGS();
+
+	/**
+	 * Construct the recovery widget.
+	 * @param InArgs The widgets arguments and attributes.
+	 */
+	void Construct(const FArguments& InArgs);
+
+	/** Requests and consumes more activities if needed and/or possible. */
+	virtual void Tick(const FGeometry& AllottedGeometry, const double InCurrentTime, const float InDeltaTime) override;
+
+	/** Returns the activity selected or null if none is selected. */
+	TSharedPtr<FConcertSessionActivity> GetSelectedActivity() const;
+
+	/** Returns the total number of activities currently stored (no filter applied). */
+	int32 GetTotalActivityNum() const { return AllActivities.Num(); }
+
+	/** Returns the number of activities shown. */
+	int32 GetDisplayedActivityNum() const { return Activities.Num(); }
+
+	/** Returns the number of activities marked as 'ignored'. */
+	int32 GetIgnoredActivityNum() const { return IgnoredActivityNum; }
+
+	/** Returns the most recent activity available, ignoring the current filter. */
+	TSharedPtr<FConcertSessionActivity> GetMostRecentActivity() const;
+
+	/** Returns true if the column names is the last one (most right one). */
+	bool IsLastColumn(const FName& ColumnId) const;
+
+	/** Clears all activities displayed. */
+	void ResetActivityList();
+
+	/** Append an activity to the view. Used to populate the view from a live session. */
+	void Append(TSharedPtr<FConcertSessionActivity> Activity);
+
+	/** Request the view to refresh. */
+	void RequestRefresh();
+
+	/** Asks the view to update the text search filter. */
+	FText UpdateTextFilter(const FText& InFilterText);
+
+private:
+	/** Generates a row widget in the table view. */
+	TSharedRef<ITableRow> OnGenerateActivityRowWidget(TSharedPtr<FConcertSessionActivity> Item, const TSharedRef<STableViewBase>& OwnerTable);
+
+	/** Fetches more activities. */
+	void FetchActivities();
+
+	/** Queries the active set of activity filter flags. */
+	EConcertActivityFilterFlags QueryActiveActivityFilters() const;
+
+	/** Invoked when the filter changes and the displayed activity list must be reevaluated. */
+	void OnActivityFilterUpdated();
+
+	/** Returns true if the specified activity passes the view filters. */
+	bool PassesFilters(const FConcertSessionActivity& Activity);
+
+	/** Invoked when the list view scrolled to fetch more data from the activity provider. */
+	void OnListViewScrolled(double InScrollOffset);
+
+	/** Invoked when the list view selection changes. */
+	void OnListViewSelectionChanged(TSharedPtr<FConcertSessionActivity> InActivity, ESelectInfo::Type SelectInfo);
+
+	/** Convert the item into a list of text element used search bars. */
+	void PopulateSearchStrings(const FConcertSessionActivity& Item, TArray<FString>& OutSearchStrings) const;
+
+	/** Returns whether the details expandable area should be visible or not. */
+	EVisibility GetDetailAreaVisibility() const { return DetailsAreaVisibility; }
+
+	/** Manages how much space is used by the 'Details' area with respect to its expansion state. */
+	SSplitter::ESizeRule GetDetailsAreaSizeRule() const { return bDetailsAreaExpanded ? SSplitter::ESizeRule::FractionOfParent : SSplitter::ESizeRule::SizeToContent; }
+
+	/** Invoked when the 'Details' expandable area is expanded/collapsed. Affects how the 'Details' area size is accounted in the splitter. */
+	void OnDetailsAreaExpansionChanged(bool bExpanded);
+
+	/** Displays what properties/objects changed during a transaction in the details area when a transaction activity is selected in the list. */
+	void DisplayTransactionDetails(const FConcertSessionActivity& Activity, const FConcertTransactionEventBase& InTransaction);
+
+	/** Displays what changed in a package in the details area when a package activity is selected in the list. */
+	void DisplayPackageDetails(const FConcertSessionActivity& Activity, int64 PackageRevision, const FConcertPackageInfo& PackageInfo);
+
+	/** Update the detail area to display the proper panel and details. */
+	void UpdateDetailArea(TSharedPtr<FConcertSessionActivity> InSelectedActivity);
+
+	/** Change the visibility of the details panels, keeping only one visible. */
+	void SetDetailsPanelVisibility(const SWidget* VisiblePanel);
+
+	/** Returns the text displayed to the user when there are no details to display. */
+	FText GetNoDetailsText() const;
+
+	/** Invoked after all widget ticks get processed. Used to check if the user scrolled or not. */
+	void OnPostTick(float);
+
+private:
+	/** List of all activities (including the filtered out ones). */
+	TArray<TSharedPtr<FConcertSessionActivity>> AllActivities;
+
+	/** List of currently displayed activities. */
+	TArray<TSharedPtr<FConcertSessionActivity>> Activities;
+
+	/** The list view widget displaying the activities. */
+	TSharedPtr<SListView<TSharedPtr<FConcertSessionActivity>>> ActivityView;
+
+	/** Used to overlay a widget over a column widget (add an extra layer above the normal one) */
+	FMakeColumnOverlayWidgetFunc MakeColumnOverlayWidgetFn;
+
+	/** Returns which text should be highlighted. */
+	TAttribute<FText> HighlightText;
+
+	/** Whether the time should be displayed as relative (9 seconds ago) or as absolute (July 10, 2019 - 10:20:10) */
+	TAttribute<ETimeFormat> TimeFormat;
+
+	/** Whether the client name column is displayed. */
+	TAttribute<EVisibility> ClientNameColumnVisibility;
+
+	/** Whether the operation column is displayed. */
+	TAttribute<EVisibility> OperationColumnVisibility;
+
+	/** Whether the package column is displayed. */
+	TAttribute<EVisibility> PackageColumnVisibility;
+
+	/** Whether the join/leave session activities are displayed. */
+	TAttribute<EVisibility> ConnectionActivitiesVisibility;
+
+	/** Whether the lock/unlock activities are displayed.*/
+	TAttribute<EVisibility> LockActivitiesVisibility;
+
+	/** Whether the package activities are displayed.*/
+	TAttribute<EVisibility> PackageActivitiesVisibility;
+
+	/** Whether the transaction activities are displayed.*/
+	TAttribute<EVisibility> TransactionActivitiesVisibility;
+
+	/** Whether the ignored activities are displayed.*/
+	TAttribute<EVisibility> IgnoredActivitiesVisibility;
+
+	/** The number of activities flagged as 'ignored' within AllActivities array.*/
+	int32 IgnoredActivityNum = 0;
+
+	/** Whether the auto-scrolling is desired (scroll to bottom automatically unless the user scroll himself somewhere). */
+	bool bAutoScrollDesired = false;
+
+	/** Whether the user is scrolling (deactivate auto-scrolling). */
+	bool bUserScrolling = false;
+
+	/** Defines which activity types are currently filtered out from the view. */
+	EConcertActivityFilterFlags ActiveFilterFlags;
+
+	/** Used to fetch more activities from an abstract source. Usually mutually exclusive with Append() function. May not be bound. */
+	FFetchActivitiesFunc FetchActivitiesFn;
+
+	/** Used to map an activity endpoint ID to a client. May not be bound. */
+	FGetActivityClientInfoFunc GetActivityUserFn;
+
+	/** Used to get the transaction event to display the selected transaction activity details. May not be bound. */
+	FGetTransactionEvent GetTransactionEventFn;
+
+	/** Used to get the package event to display the selected package activity details. May not be bound. */
+	FGetPackageEvent GetPackageEventFn;
+
+	/** The number of activities to request when scrolling down to request activities on demand. Used with FetchActivitiesFn() for paging. */
+	static constexpr int32 ActivitiesPerRequest = 128;
+
+	/** The current desired amount of activities to display. Used for paging with FetchActivitiesFn(). It will grow when the user scrolls down. */
+	int32 DesiredActivitiesCount = ActivitiesPerRequest;
+
+	/** True once the activity provider function (FetchActivitiesFn) returns true. */
+	bool bAllActivitiesFetched = false;
+
+	/** Utility class used to tokenize and match text displayed in the list view. */
+	TSharedPtr<TTextFilter<const FConcertSessionActivity&>> SearchTextFilter;
+
+	/** The expandable area under which the activity details are displayed. */
+	TSharedPtr<SExpandableArea> ExpandableDetails;
+
+	/** The widget displaying transaction details (if detail area is expanded) when a transaction activity is selected. */
+	TSharedPtr<SUndoHistoryDetails> TransactionDetailsPanel;
+
+	/** The widget displaying package details (if detail area is expanded) when a package activity is selected.*/
+	TSharedPtr<SPackageDetails> PackageDetailsPanel;
+
+	/** The widget saying they are no detail available (if detail area is expanded) when the activity has no details or no activity is selected. */
+	TSharedPtr<SBox> NoDetailsPanel;
+
+	/** The widget displayed when details of a partially sync activity (if detail area is expanded) are being fetched from the server. */
+	TSharedPtr<SWidget> LoadingDetailsPanel;
+
+	/** Indicate if the details area should be displayed or not. */
+	EVisibility DetailsAreaVisibility = EVisibility::Collapsed;
+
+	/** Keeps the expanded status of the details area. */
+	bool bDetailsAreaExpanded = false;
+
+	/** Indicate whether the activity list view was scrolled during a frame. Used to correctly detect auto-scrolling in special cases. */
+	bool bActivityViewScrolled = false;
+};
+
+/**
+ * Manages the various options exposed by SConcertSessionActivities widget such as enabling/disabling filtering,
+ * changing the time format, etc.
+ */
+class CONCERTSHAREDSLATE_API FConcertSessionActivitiesOptions : public TSharedFromThis<FConcertSessionActivitiesOptions>
+{
+public:
+	/** Returns a Menu widgets containing the available options. */
+	TSharedRef<SWidget> MakeMenuWidget();
+
+	/** Makes a standard View Options widget, displaying the eye ball icon and showing the possible options. */
+	TSharedRef<SWidget> MakeViewOptionsWidget();
+
+	/** Makes a widgets saying how many operations are shown with respect to those filtered out. */
+	TSharedRef<SWidget> MakeDisplayedActivityCountWidget(TAttribute<int32> Total, TAttribute<int32> Displayed);
+
+	/** Makes a status bar widget displaying the activity shown and the standard view options button. */
+	TSharedRef<SWidget> MakeStatusBar(TAttribute<int32> Total, TAttribute<int32> Displayed);
+
+	/** Returns the time format option. */
+	SConcertSessionActivities::ETimeFormat GetTimeFormat() const { return bDisplayRelativeTime ? SConcertSessionActivities::ETimeFormat::Relative : SConcertSessionActivities::ETimeFormat::Absolute; }
+
+	/** Returns whether the connection activities are listed in SConcertSessionActivities. */
+	EVisibility GetConnectionActivitiesVisibility() const { return bDisplayConnectionActivities ? EVisibility::Visible : EVisibility::Hidden; }
+
+	/** Returns whether the lock activities are listed in SConcertSessionActivities. */
+	EVisibility GetLockActivitiesVisibility() const { return bDisplayLockActivities ? EVisibility::Visible : EVisibility::Hidden; }
+
+	/** Returns whether the package activities are listed in SConcertSessionActivities. */
+	EVisibility GetPackageActivitiesVisibility() const { return bDisplayPackageActivities ? EVisibility::Visible : EVisibility::Hidden; }
+
+	/** Returns whether the transaction activities are listed in SConcertSessionActivities. */
+	EVisibility GetTransactionActivitiesVisibility() const { return bDisplayTransactionActivities ? EVisibility::Visible : EVisibility::Hidden; }
+
+	/** Returns whether the ignored activities are listed in SConcertSessionActivities. */
+	EVisibility GetIgnoredActivitiesVisibility() const { return bDisplayIgnoredActivities ? EVisibility::Visible : EVisibility::Hidden; }
+
+	/** Invoked when an options is togged from the displayed menu widget. */
+	void OnOptionToggled(const FName CheckBoxId);
+
+	/** Enables the 'connection activity' filter check box. (Show Connection Activities). */
+	bool bEnableConnectionActivityFiltering = true;
+
+	/** Enables the 'lock activity' filter check box. (Show Lock Activities).*/
+	bool bEnableLockActivityFiltering = true;
+
+	/** Enables the 'package activity' filter check box (Show Package Activities). */
+	bool bEnablePackageActivityFiltering = true;
+
+	/** Enables the 'transaction activity' filter check box (Show Transaction Activities). */
+	bool bEnableTransactionActivityFiltering = true;
+
+	/** Enables the 'ignored activity' filter check box (Show Unrecoverable Activities). */
+	bool bEnableIgnoredActivityFiltering = false;
+
+	/** Controls whether the time is displayed as absolute or relative. (Display Relative Time). */
+	bool bDisplayRelativeTime = true;
+
+	/** If connection filtering is enabled, controls whether connection activities are filtered out. */
+	bool bDisplayConnectionActivities = false;
+
+	/** If lock filtering is enabled, controls whether lock activities are filtered out. */
+	bool bDisplayLockActivities = false;
+
+	/** If package filtering is enabled, controls whether package activities are filtered out. */
+	bool bDisplayPackageActivities = true;
+
+	/** If transaction filtering is enabled, controls whether transaction activities are filtered out.*/
+	bool bDisplayTransactionActivities = true;
+
+	/** If ignored activity filtering is enabled, controls whether ignored activities are filtered out.*/
+	bool bDisplayIgnoredActivities = false;
+};
diff --git a/Engine/Plugins/Developer/Concert/ConcertUI/ConcertSharedSlate/Source/ConcertSharedSlate/Public/Session/Browser/ConcertBrowserUtils.h b/Engine/Plugins/Developer/Concert/ConcertUI/ConcertSharedSlate/Source/ConcertSharedSlate/Public/Session/Browser/ConcertBrowserUtils.h
new file mode 100644
--- /dev/null
+++ b/Engine/Plugins/Developer/Concert/ConcertUI/ConcertSharedSlate/Source/ConcertSharedSlate/Public/Session/Browser/ConcertBrowserUtils.h
@@ -0,0 +1,100 @@
+﻿// Copyright Epic Games, Inc. All Rights Reserved.
+
+#pragma once
+
+#include "CoreMinimal.h"
+
+#include "ConcertMessageData.h"
+
+#include "EditorFontGlyphs.h"
+#include "SNegativeActionButton.h"
+#include "SPositiveActionButton.h"
+#include "Styling/AppStyle.h"
+#include "Internationalization/Regex.h"
+#include "Layout/Visibility.h"
+#include "Misc/Attribute.h"
+#include "Misc/AsyncTaskNotification.h"
+#include "Styling/SlateColor.h"
+#include "Widgets/Images/SImage.h"
+#include "Widgets/Input/SButton.h"
+#include "Widgets/Layout/SBorder.h"
+#include "Widgets/Text/STextBlock.h"
+
+namespace ConcertBrowserUtils
+{
+	// Defines the sessions list view column tag names.
+	static const FName IconColName(TEXT("Icon"));
+	static const FName SessionColName(TEXT("Session"));
+	static const FName ServerColName(TEXT("Server"));
+
+	// Name of the filter box in the View option.
+	static const FName ActiveSessionsCheckBoxMenuName(TEXT("ActiveSessions"));
+	static const FName ArchivedSessionsCheckBoxMenuName(TEXT("ArchivedSessions"));
+	static const FName DefaultServerCheckBoxMenuName(TEXT("DefaultServer"));
+
+	// The awesome font used to pick the icon displayed in the session list view 'Icon' column.
+	static const FName IconColumnFontName(TEXT("FontAwesome.9"));
+
+	/** Utility function used to create buttons displaying only an icon (using a brush) */
+	inline TSharedRef<SWidget> MakeIconButton(const TAttribute<const FSlateBrush*>& Icon, const TAttribute<FText>& Tooltip, const TAttribute<bool>& EnabledAttribute, const FOnClicked& OnClicked, const TAttribute<EVisibility>& Visibility = EVisibility::Visible)
+	{
+		return SNew(SButton)
+			.OnClicked(OnClicked)
+			.ButtonStyle(FAppStyle::Get(), TEXT("SimpleButton"))
+			.ToolTipText(Tooltip)
+			.Visibility(Visibility)
+			.IsEnabled(EnabledAttribute)
+			.VAlign(VAlign_Center)
+			.HAlign(HAlign_Center)
+			[
+				SNew(SImage)
+				.Image(Icon)
+				.ColorAndOpacity(FSlateColor::UseForeground())
+			];
+	}
+
+	/** Utility function used to create buttons displaying a positive action*/
+	inline TSharedRef<SWidget> MakePositiveActionButton(const TAttribute<const FSlateBrush*>& Icon, const TAttribute<FText>& Tooltip, const TAttribute<bool>& EnabledAttribute, const FOnClicked& OnClicked, const TAttribute<EVisibility>& Visibility = EVisibility::Visible)
+	{
+		return SNew(SPositiveActionButton)
+			.OnClicked(OnClicked)
+			.ToolTipText(Tooltip)
+			.Visibility(Visibility)
+			.IsEnabled(EnabledAttribute)
+			.Icon(Icon);
+	}
+
+	/** Utility function used to create buttons displaying a negative action */
+	inline TSharedRef<SWidget> MakeNegativeActionButton(const TAttribute<const FSlateBrush*>& Icon, const TAttribute<FText>& Tooltip, const TAttribute<bool>& EnabledAttribute, const FOnClicked& OnClicked, const TAttribute<EVisibility>& Visibility = EVisibility::Visible)
+	{
+		return SNew(SNegativeActionButton)
+			.OnClicked(OnClicked)
+			.ToolTipText(Tooltip)
+			.Visibility(Visibility)
+			.IsEnabled(EnabledAttribute)
+			.Icon(Icon);
+	}
+
+	/** Returns the tooltip shown when hovering the triangle with an exclamation icon when a server doesn't validate the version requirements. */
+	inline FText GetServerVersionIgnoredTooltip()
+	{
+#define LOCTEXT_NAMESPACE "SConcertBrowser"
+		return LOCTEXT("ServerIgnoreSessionRequirementsTooltip", "Careful this server won't verify that you have the right requirements before you join a session");
+#undef LOCTEXT_NAMESPACE
+	}
+
+	/** Create a widget displaying the triangle with an exclamation icon in case the server flags include IgnoreSessionRequirement. */
+	inline TSharedRef<SWidget> MakeServerVersionIgnoredWidget(EConcertServerFlags InServerFlags)
+	{
+		return SNew(SBorder)
+			.BorderImage(FAppStyle::Get().GetBrush("NoBorder"))
+			.ColorAndOpacity(FAppStyle::Get().GetWidgetStyle<FButtonStyle>("FlatButton.Warning").Normal.TintColor.GetSpecifiedColor())
+			[
+				SNew(STextBlock)
+				.Font(FAppStyle::Get().GetFontStyle("FontAwesome.9"))
+				.Text(FEditorFontGlyphs::Exclamation_Triangle)
+				.ToolTipText(GetServerVersionIgnoredTooltip())
+				.Visibility((InServerFlags & EConcertServerFlags::IgnoreSessionRequirement) != EConcertServerFlags::None ? EVisibility::Visible : EVisibility::Collapsed)
+			];
+	}
+}
diff --git a/Engine/Plugins/Developer/Concert/ConcertUI/ConcertSharedSlate/Source/ConcertSharedSlate/Public/Session/Browser/ConcertSessionBrowserSettings.h b/Engine/Plugins/Developer/Concert/ConcertUI/ConcertSharedSlate/Source/ConcertSharedSlate/Public/Session/Browser/ConcertSessionBrowserSettings.h
new file mode 100644
--- /dev/null
+++ b/Engine/Plugins/Developer/Concert/ConcertUI/ConcertSharedSlate/Source/ConcertSharedSlate/Public/Session/Browser/ConcertSessionBrowserSettings.h
@@ -0,0 +1,26 @@
+// Copyright Epic Games, Inc. All Rights Reserved.
+
+#pragma once
+
+#include "UObject/Object.h"
+#include "ConcertSessionBrowserSettings.generated.h"
+
+/** Serializes the multi-user session browser settings like the active filters. */
+UCLASS(config=Editor)
+class CONCERTSHAREDSLATE_API UConcertSessionBrowserSettings : public UObject
+{
+	GENERATED_BODY()
+
+public:
+	UConcertSessionBrowserSettings() {}
+
+	UPROPERTY(config, EditAnywhere, Category="Multi-User Session Browser")
+	bool bShowActiveSessions = true;
+
+	UPROPERTY(config, EditAnywhere, Category="Multi-User Session Browser")
+	bool bShowArchivedSessions = true;
+
+	UPROPERTY(config, EditAnywhere, Category="Multi-User Session Browser")
+	bool bShowDefaultServerSessionsOnly = false;
+};
+
diff --git a/Engine/Plugins/Developer/Concert/ConcertUI/ConcertSharedSlate/Source/ConcertSharedSlate/Public/Session/Browser/ConcertSessionItem.h b/Engine/Plugins/Developer/Concert/ConcertUI/ConcertSharedSlate/Source/ConcertSharedSlate/Public/Session/Browser/ConcertSessionItem.h
new file mode 100644
--- /dev/null
+++ b/Engine/Plugins/Developer/Concert/ConcertUI/ConcertSharedSlate/Source/ConcertSharedSlate/Public/Session/Browser/ConcertSessionItem.h
@@ -0,0 +1,49 @@
+﻿// Copyright Epic Games, Inc. All Rights Reserved.
+
+#pragma once
+
+#include "CoreMinimal.h"
+#include "ConcertMessageData.h"
+
+/** Signal emitted when a session name text field should enter in edit mode. */
+DECLARE_MULTICAST_DELEGATE(FOnBeginEditConcertSessionNameRequest)
+
+/**
+ * Item displayed in the session list view.
+ */
+class CONCERTSHAREDSLATE_API FConcertSessionItem
+{
+public:
+	enum class EType
+	{
+		None,
+		NewSession,      // Editable item to enter a session name and a pick a server.
+		RestoreSession,  // Editable item to name the restored session.
+		SaveSession,     // Editable item to name the archive.
+		ActiveSession,   // Read-only item representing an active session.
+		ArchivedSession, // Read-only item representing an archived session.
+	};
+
+	FConcertSessionItem(EType Type, const FString& InSessionName, const FGuid& InSessionId, const FString& InServerName, const FGuid& InServerEndpoint, EConcertServerFlags InServerFlags)
+		: Type(Type)
+		, ServerAdminEndpointId(InServerEndpoint)
+		, SessionId(InSessionId)
+		, SessionName(InSessionName)
+		, ServerName(InServerName)
+		, ServerFlags(InServerFlags)
+	{
+	}
+
+	bool operator==(const FConcertSessionItem& Other) const
+	{
+		return Type == Other.Type && ServerAdminEndpointId == Other.ServerAdminEndpointId && SessionId == Other.SessionId;
+	}
+
+	EType Type = EType::None;
+	FGuid ServerAdminEndpointId;
+	FGuid SessionId;
+	FString SessionName;
+	FString ServerName;
+	FOnBeginEditConcertSessionNameRequest OnBeginEditSessionNameRequest; // Emitted when user press 'F2' or select 'Rename' from context menu.
+	EConcertServerFlags ServerFlags;
+};
\ No newline at end of file
diff --git a/Engine/Plugins/Developer/Concert/ConcertUI/ConcertSharedSlate/Source/ConcertSharedSlate/Public/Session/Browser/IConcertSessionBrowserController.h b/Engine/Plugins/Developer/Concert/ConcertUI/ConcertSharedSlate/Source/ConcertSharedSlate/Public/Session/Browser/IConcertSessionBrowserController.h
new file mode 100644
--- /dev/null
+++ b/Engine/Plugins/Developer/Concert/ConcertUI/ConcertSharedSlate/Source/ConcertSharedSlate/Public/Session/Browser/IConcertSessionBrowserController.h
@@ -0,0 +1,71 @@
+﻿// Copyright Epic Games, Inc. All Rights Reserved.
+
+#pragma once
+
+#include "CoreMinimal.h"
+#include "ConcertMessageData.h"
+
+/**
+ * Runs and cache network queries for the UI. In the model-view-controller pattern, this class acts like the controller. Its purpose
+ * is to keep the UI code as decoupled as possible from the API used to query it. It encapsulate the asynchronous code and provide a
+ * simpler API to the UI.
+ */
+class CONCERTSHAREDSLATE_API IConcertSessionBrowserController
+{
+public:
+
+	struct FActiveSessionInfo
+	{
+		FConcertServerInfo ServerInfo;
+		FConcertSessionInfo SessionInfo;
+		TArray<FConcertSessionClientInfo> Clients;
+
+		FActiveSessionInfo(FConcertServerInfo ServerInfo, FConcertSessionInfo SessionInfo, TArray<FConcertSessionClientInfo> Clients = {})
+			: ServerInfo(MoveTemp(ServerInfo))
+			, SessionInfo(MoveTemp(SessionInfo))
+			, Clients(MoveTemp(Clients))
+		{}
+	};
+
+	struct FArchivedSessionInfo
+	{
+		FConcertServerInfo ServerInfo;
+		FConcertSessionInfo SessionInfo;
+
+		FArchivedSessionInfo(FConcertServerInfo ServerInfo, FConcertSessionInfo SessionInfo)
+			: ServerInfo(MoveTemp(ServerInfo))
+			, SessionInfo(MoveTemp(SessionInfo))
+		{}
+	};
+
+public:
+
+	/** Returns the latest list of server known to this controller. */
+	virtual TArray<FConcertServerInfo> GetServers() const = 0;
+
+	/** Returns the latest list of active sessions known to this controller. */
+	virtual TArray<FActiveSessionInfo> GetActiveSessions() const = 0;
+
+	/** Returns the latest list of archived sessions known to this controller. */
+	virtual TArray<FArchivedSessionInfo> GetArchivedSessions() const = 0;
+
+	/** Returns the active sessions info corresponding to the specified parameters. Used to display the sessions details. */
+	virtual TOptional<FConcertSessionInfo> GetActiveSessionInfo(const FGuid& AdminEndpoint, const FGuid& SessionId) const = 0;
+
+	/** Returns the archived sessions info corresponding to the specified parameters. Used to display the sessions details. */
+	virtual TOptional<FConcertSessionInfo> GetArchivedSessionInfo(const FGuid& AdminEndpoint, const FGuid& SessionId) const = 0;
+
+	virtual void CreateSession(const FGuid& ServerAdminEndpointId, const FString& SessionName) = 0;
+	virtual void ArchiveSession(const FGuid& ServerAdminEndpointId, const FGuid& SessionId, const FString& ArchiveName, const FConcertSessionFilter& SessionFilter) = 0;
+	virtual void RestoreSession(const FGuid& ServerAdminEndpointId, const FGuid& SessionId, const FString& RestoredName, const FConcertSessionFilter& SessionFilter) = 0;
+	virtual void RenameActiveSession(const FGuid& ServerAdminEndpointId, const FGuid& SessionId, const FString& NewName) = 0;
+	virtual void RenameArchivedSession(const FGuid& ServerAdminEndpointId, const FGuid& SessionId, const FString& NewName) = 0;
+	virtual bool CanRenameActiveSession(const FGuid& ServerAdminEndpointId, const FGuid& SessionId) const = 0;
+	virtual bool CanRenameArchivedSession(const FGuid& ServerAdminEndpointId, const FGuid& SessionId) const = 0;
+	virtual void DeleteActiveSession(const FGuid& ServerAdminEndpointId, const FGuid& SessionId) = 0;
+	virtual void DeleteArchivedSession(const FGuid& ServerAdminEndpointId, const FGuid& SessionId) = 0;
+	virtual bool CanDeleteActiveSession(const FGuid& ServerAdminEndpointId, const FGuid& SessionId) const = 0;
+	virtual bool CanDeleteArchivedSession(const FGuid& ServerAdminEndpointId, const FGuid& SessionId) const = 0;
+	
+	virtual ~IConcertSessionBrowserController() = default;
+};
\ No newline at end of file
diff --git a/Engine/Plugins/Developer/Concert/ConcertUI/ConcertSharedSlate/Source/ConcertSharedSlate/Public/Session/Browser/SConcertSessionBrowser.h b/Engine/Plugins/Developer/Concert/ConcertUI/ConcertSharedSlate/Source/ConcertSharedSlate/Public/Session/Browser/SConcertSessionBrowser.h
new file mode 100644
--- /dev/null
+++ b/Engine/Plugins/Developer/Concert/ConcertUI/ConcertSharedSlate/Source/ConcertSharedSlate/Public/Session/Browser/SConcertSessionBrowser.h
@@ -0,0 +1,226 @@
+﻿// Copyright Epic Games, Inc. All Rights Reserved.
+
+#pragma once
+
+#include "CoreMinimal.h"
+#include "Misc/TextFilter.h"
+#include "SlateFwd.h"
+#include "Widgets/DeclarativeSyntaxSupport.h"
+#include "Widgets/SCompoundWidget.h"
+#include "Widgets/Views/SHeaderRow.h"
+#include "Widgets/Views/SListView.h"
+
+class FConcertSessionItem;
+class FExtender;
+class IConcertSessionBrowserController;
+class ITableRow;
+class UConcertSessionBrowserSettings;
+class SGridPanel;
+class STableViewBase;
+
+struct FConcertSessionClientInfo;
+struct FConcertSessionInfo;
+
+DECLARE_DELEGATE_RetVal_OneParam(TSharedRef<SWidget>, FExtendSessionTable, TSharedRef<SWidget> /* TableView */);
+DECLARE_DELEGATE_OneParam(FExtenderDelegate, FExtender&);
+DECLARE_DELEGATE_TwoParams(FExtendSessionContextMenu, TSharedPtr<FConcertSessionItem>, FExtender&)
+DECLARE_DELEGATE_OneParam(FSessionDelegate, TSharedPtr<FConcertSessionItem> /*Session*/);
+DECLARE_DELEGATE_RetVal_OneParam(bool, FCanRemoveSession, TSharedPtr<FConcertSessionItem> /* SessionItem */);
+
+/**
+ * Enables the user to browse/search/filter/sort active and archived sessions, create new session,
+ * archive active sessions, restore archived sessions, join a session and open the settings dialog.
+ */
+class CONCERTSHAREDSLATE_API SConcertSessionBrowser : public SCompoundWidget
+{
+public:
+	
+	struct CONCERTSHAREDSLATE_API ControlButtonExtensionHooks
+	{
+		/** Contains: Create Session */
+		static const FName BeforeSeparator;
+		/** Just separates the two */
+		static const FName Separator;
+		/** Contains: Restore, Archive, Delete */
+		static const FName AfterSeparator;
+	};
+
+	struct CONCERTSHAREDSLATE_API SessionContextMenuExtensionHooks
+	{
+		/** Contains: Archive (ActiveSession), Restore (ArchivedSession), Rename, Delete  */
+		static const FName ManageSession;
+	};
+	
+	SLATE_BEGIN_ARGS(SConcertSessionBrowser) { }
+
+	/** Optional name of the default session - relevant for highlighting */
+	SLATE_ATTRIBUTE(FString, DefaultSessionName)
+	/** Optional url of the default server - relevant for highlighting */
+	SLATE_ATTRIBUTE(FString, DefaultServerURL)
+
+	/** Used during construction to override how the session table view is created, e.g. to embed it into an overlay */
+	SLATE_EVENT(FExtendSessionTable, ExtendSessionTable)
+	/** Extends the buttons to the left of the search bar */
+	SLATE_EVENT(FExtenderDelegate, ExtendControllButtons)
+	/** Extends the menu when the user right-clicks a session */
+	SLATE_EVENT(FExtendSessionContextMenu, ExtendSessionContextMenu)
+	/** Custom slot placed to the right of the search bar */
+	SLATE_NAMED_SLOT(FArguments, RightOfSearchBar)
+	
+	/** Called when this session is clicked */
+	SLATE_EVENT(FSessionDelegate, OnSessionClicked)
+	/** Called when this session is double-clicked */
+	SLATE_EVENT(FSessionDelegate, OnSessionDoubleClicked)
+	/** Called after a user has requested to delete a session */
+	SLATE_EVENT(FSessionDelegate, OnRequestedDeleteSession)
+
+	/** Ask the user to confirm archiving - most obvious implementation is showing a dialog box */
+	SLATE_EVENT(FCanRemoveSession, CanDeleteArchivedSession)
+	/** Ask the user to confirm deleting - most obvious implementation is showing a dialog box */
+	SLATE_EVENT(FCanRemoveSession, CanDeleteActiveSession)
+	
+	SLATE_END_ARGS();
+
+	/**
+	* Constructs the Browser.
+	* @param InArgs The Slate argument list.
+	* @param InController The controller used to send queries from the UI - represents controller in mode-view-controller pattern.
+	* @param[in,out] InSearchText The text to set in the search box and to remember (as output). Cannot be null.
+	*/
+	void Construct(const FArguments& InArgs, TSharedRef<IConcertSessionBrowserController> InController, TSharedPtr<FText> InSearchText);
+
+	//~ Begin SWidget Interface
+	virtual void Tick(const FGeometry& AllottedGeometry, const double InCurrentTime, const float InDeltaTime) override;
+	virtual FReply OnKeyDown(const FGeometry& MyGeometry, const FKeyEvent& InKeyEvent) override;
+	//~ End SWidget Interface
+	
+	void RefreshSessionList();
+	TArray<TSharedPtr<FConcertSessionItem>> GetSessions() const { return Sessions; }
+	TArray<TSharedPtr<FConcertSessionItem>> GetSelectedItems() const { return SessionsView->GetSelectedItems(); }
+
+	bool IsNewButtonEnabled() const { return IsNewButtonEnabledInternal(); }
+	bool IsRestoreButtonEnabled() const { return IsRestoreButtonEnabledInternal(); }
+	bool IsArchiveButtonEnabled() const { return IsArchiveButtonEnabledInternal(); }
+	bool IsRenameButtonEnabled() const { return IsRenameButtonEnabledInternal(); }
+	bool IsDeleteButtonEnabled() const { return IsDeleteButtonEnabledInternal(); }
+	
+	/** Adds row for creating new session. Exposed for other widgets, e.g. discovery overlay to create a new session. */
+	void InsertNewSessionEditableRow() { InsertNewSessionEditableRowInternal(); }
+	/** Creates row under the given (archived) session with which session can be restored. */
+	void InsertRestoreSessionAsEditableRow(const TSharedPtr<FConcertSessionItem>& ArchivedItem) { InsertRestoreSessionAsEditableRowInternal(ArchivedItem); };
+	
+private:
+	
+	// Layout the 'session|details' split view.
+	TSharedRef<SWidget> MakeBrowserContent(const FArguments& InArgs);
+
+	// Layout the sessions view and controls.
+	TSharedRef<SWidget> MakeControlBar(const FArguments& InArgs);
+	TSharedRef<SWidget> MakeButtonBar(const FArguments& InArgs);
+	TSharedRef<SWidget> MakeSessionTableView();
+	TSharedRef<SWidget> MakeSessionViewOptionsBar();
+	TSharedRef<ITableRow> OnGenerateSessionRowWidget(TSharedPtr<FConcertSessionItem> Item, const TSharedRef<STableViewBase>& OwnerTable);
+
+	// Creates row widgets for session list view, validates user inputs and forward user requests for processing to a delegate function implemented by this class.
+	TSharedRef<ITableRow> MakeActiveSessionRowWidget(const TSharedPtr<FConcertSessionItem>& ActiveItem, const TSharedRef<STableViewBase>& OwnerTable);
+	TSharedRef<ITableRow> MakeArchivedSessionRowWidget(const TSharedPtr<FConcertSessionItem>& ArchivedItem, const TSharedRef<STableViewBase>& OwnerTable);
+	TSharedRef<ITableRow> MakeNewSessionRowWidget(const TSharedPtr<FConcertSessionItem>& NewItem, const TSharedRef<STableViewBase>& OwnerTable);
+	TSharedRef<ITableRow> MakeRestoreSessionRowWidget(const TSharedPtr<FConcertSessionItem>& RestoreItem, const TSharedRef<STableViewBase>& OwnerTable);
+	TSharedRef<ITableRow> MakeSaveSessionRowWidget(const TSharedPtr<FConcertSessionItem>& ArchivedItem, const TSharedRef<STableViewBase>& OwnerTable);
+
+	// Creates the contextual menu when right clicking a session list view row.
+	TSharedPtr<SWidget> MakeContextualMenu();
+
+	// The buttons above the session view.
+	bool IsNewButtonEnabledInternal() const;
+	bool IsRestoreButtonEnabledInternal() const;
+	bool IsArchiveButtonEnabledInternal() const;
+	bool IsRenameButtonEnabledInternal() const;
+	bool IsDeleteButtonEnabledInternal() const;
+	FReply OnNewButtonClicked();
+	FReply OnRestoreButtonClicked();
+	FReply OnArchiveButtonClicked();
+	FReply OnDeleteButtonClicked();
+	void OnBeginEditingSessionName(TSharedPtr<FConcertSessionItem> Item);
+
+	// Manipulates the sessions view (the array and the UI).
+	void OnSessionSelectionChanged(TSharedPtr<FConcertSessionItem> SelectedSession, ESelectInfo::Type SelectInfo);
+	void InsertNewSessionEditableRowInternal();
+	void InsertRestoreSessionAsEditableRowInternal(const TSharedPtr<FConcertSessionItem>& ArchivedItem);
+	void InsertArchiveSessionAsEditableRow(const TSharedPtr<FConcertSessionItem>& ActiveItem);
+	void InsertEditableSessionRow(TSharedPtr<FConcertSessionItem> EditableItem, TSharedPtr<FConcertSessionItem> ParentItem);
+	void RemoveSessionRow(const TSharedPtr<FConcertSessionItem>& Item);
+
+	// Sessions sorting. (Sorts the session view)
+	EColumnSortMode::Type GetColumnSortMode(const FName ColumnId) const;
+	EColumnSortPriority::Type GetColumnSortPriority(const FName ColumnId) const;
+	void OnColumnSortModeChanged(const EColumnSortPriority::Type SortPriority, const FName& ColumnId, const EColumnSortMode::Type InSortMode);
+	void SortSessionList();
+	void EnsureEditableParentChildOrder();
+
+	// Sessions filtering. (Filters the session view)
+	void OnSearchTextChanged(const FText& InFilterText);
+	void OnSearchTextCommitted(const FText& InFilterText, ETextCommit::Type CommitType);
+	void OnFilterMenuChecked(const FName MenuName);
+	void PopulateSearchStrings(const FConcertSessionItem& Item, TArray<FString>& OutSearchStrings) const;
+	bool IsFilteredOut(const FConcertSessionItem& Item) const;
+	FText HighlightSearchText() const;
+
+	// Passes the user requests to FConcertBrowserController.
+	void RequestCreateSession(const TSharedPtr<FConcertSessionItem>& NewItem);
+	void RequestArchiveSession(const TSharedPtr<FConcertSessionItem>& ActiveItem, const FString& ArchiveName);
+	void RequestRestoreSession(const TSharedPtr<FConcertSessionItem>& RestoreItem, const FString& SessionName);
+	void RequestRenameSession(const TSharedPtr<FConcertSessionItem>& RenamedItem, const FString& NewName);
+	void RequestDeleteSession(const TSharedPtr<FConcertSessionItem>& DeletedItem);
+
+	TSharedPtr<IConcertSessionBrowserController> GetController() const;
+
+private:
+	
+	// Gives access to the concert data (servers, sessions, clients, etc).
+	TWeakPtr<IConcertSessionBrowserController> Controller;
+
+	// Keeps persistent user preferences, like the filters.
+	UConcertSessionBrowserSettings* PersistentSettings = nullptr;
+
+	/** Optional default session name - relevant for highlighting */
+	TAttribute<FString> DefaultSessionName;
+	/** Optional default server url - relevant for highlighting */
+	TAttribute<FString> DefaultServerUrl;
+	
+	FExtendSessionContextMenu ExtendSessionContextMenu;
+	FSessionDelegate OnSessionClicked;
+	FSessionDelegate OnSessionDoubleClicked;
+	FSessionDelegate OnRequestedDeleteSession;
+	FCanRemoveSession CanDeleteArchivedSession;
+	FCanRemoveSession CanDeleteActiveSession;
+
+	// The items displayed in the session list view. It might be filtered and sorted compared to the full list hold by the controller.
+	TArray<TSharedPtr<FConcertSessionItem>> Sessions;
+
+	// The session list view.
+	TSharedPtr<SListView<TSharedPtr<FConcertSessionItem>>> SessionsView;
+
+	// The item corresponding to a row used to create/archive/restore a session. There is only one at the time
+	TSharedPtr<FConcertSessionItem> EditableSessionRow;
+	TSharedPtr<FConcertSessionItem> EditableSessionRowParent; // For archive/restore, indicate which element is archived or restored.
+
+	// Sorting.
+	EColumnSortMode::Type PrimarySortMode = EColumnSortMode::None;
+	EColumnSortMode::Type SecondarySortMode = EColumnSortMode::None;
+	FName PrimarySortedColumn;
+	FName SecondarySortedColumn;
+
+	// Filtering.
+	TSharedPtr<SSearchBox> SearchBox;
+	TSharedPtr<TTextFilter<const FConcertSessionItem&>> SearchTextFilter;
+	TSharedPtr<FText> SearchedText;
+	bool bRefreshSessionFilter = true;
+	FString LastDefaultServerUrl;
+
+	// Selected Session Details.
+	TSharedPtr<SBorder> SessionDetailsView;
+	TSharedPtr<SExpandableArea> DetailsArea;
+	TArray<TSharedPtr<FConcertSessionClientInfo>> Clients;
+	TSharedPtr<SExpandableArea> ClientsArea;
+};
diff --git a/Engine/Plugins/Developer/Concert/ConcertUI/ConcertSharedSlate/Source/ConcertSharedSlate/Public/Session/History/AbstractSessionHistoryController.h b/Engine/Plugins/Developer/Concert/ConcertUI/ConcertSharedSlate/Source/ConcertSharedSlate/Public/Session/History/AbstractSessionHistoryController.h
new file mode 100644
--- /dev/null
+++ b/Engine/Plugins/Developer/Concert/ConcertUI/ConcertSharedSlate/Source/ConcertSharedSlate/Public/Session/History/AbstractSessionHistoryController.h
@@ -0,0 +1,45 @@
+﻿// Copyright Epic Games, Inc. All Rights Reserved.
+
+#pragma once
+
+#include "CoreMinimal.h"
+#include "Async/Future.h"
+
+struct FConcertClientInfo;
+struct FConcertSyncTransactionEvent;
+struct FConcertSyncPackageEventMetaData;
+struct FConcertSessionActivity;
+class IConcertSyncClient;
+class SSessionHistory;
+
+/** Acts as controller of SSessionHistory in model-view-component pattern. */
+class CONCERTSHAREDSLATE_API FAbstractSessionHistoryController : public TSharedFromThis<FAbstractSessionHistoryController>
+{
+public:
+
+	FAbstractSessionHistoryController(FName PackageFilter = EName::None);
+	virtual ~FAbstractSessionHistoryController() = default;
+	
+	/** Fetches activities from the server and updates the list view. */
+	void ReloadActivities();
+	
+	const TSharedRef<SSessionHistory>& GetSessionHistory() const { return SessionHistory; }
+
+protected:
+
+	/** Gets the latest activities*/
+	virtual void GetActivities(int64 MaximumNumberOfActivities, TMap<FGuid, FConcertClientInfo>& EndpointClientInfoMap, TArray<FConcertSessionActivity>& FetchedActivities) const = 0;
+	
+	/** Returns the specified package event (without the package data itself) if available. */
+	virtual bool GetPackageEvent(const FConcertSessionActivity& Activity, FConcertSyncPackageEventMetaData& OutPackageEvent) const = 0;
+
+	/** Returns the specified package event if available. */
+	virtual TFuture<TOptional<FConcertSyncTransactionEvent>> GetTransactionEvent(const FConcertSessionActivity& Activity) const = 0;
+	
+private:
+
+	/** The widget being managed */
+	TSharedRef<SSessionHistory> SessionHistory;
+
+	TSharedRef<SSessionHistory> MakeSessionHistory(FName PackageFilter) const;
+};
\ No newline at end of file
diff --git a/Engine/Plugins/Developer/Concert/ConcertUI/ConcertSharedSlate/Source/ConcertSharedSlate/Public/Session/History/SSessionHistory.h b/Engine/Plugins/Developer/Concert/ConcertUI/ConcertSharedSlate/Source/ConcertSharedSlate/Public/Session/History/SSessionHistory.h
new file mode 100644
--- /dev/null
+++ b/Engine/Plugins/Developer/Concert/ConcertUI/ConcertSharedSlate/Source/ConcertSharedSlate/Public/Session/History/SSessionHistory.h
@@ -0,0 +1,77 @@
+// Copyright Epic Games, Inc. All Rights Reserved.
+
+#pragma once
+
+#include "CoreMinimal.h"
+#include "Session/Activity/SConcertSessionActivities.h"
+#include "Widgets/DeclarativeSyntaxSupport.h"
+#include "Widgets/SCompoundWidget.h"
+
+class SSearchBox;
+class SConcertSessionActivities;
+class FConcertSessionActivitiesOptions;
+
+class CONCERTSHAREDSLATE_API SSessionHistory : public SCompoundWidget
+{
+public:
+
+	/** Maximum number of activities displayed on screen. */ 
+	static constexpr int64 MaximumNumberOfActivities = 1000;
+	
+	SLATE_BEGIN_ARGS(SSessionHistory) {}
+		SLATE_ARGUMENT(FName, PackageFilter)
+		SLATE_EVENT(SConcertSessionActivities::FGetPackageEvent, GetPackageEvent)
+		SLATE_EVENT(SConcertSessionActivities::FGetTransactionEvent, GetTransactionEvent)
+	SLATE_END_ARGS()
+
+	/**
+	* Constructs the Session History widget.
+	*
+	* @param InArgs The Slate argument list.
+	* @param ConstructUnderMajorTab The major tab which will contain the Session History widget.
+	* @param InConcertSyncClient Pointer on the concert sync client.
+	*/
+	void Construct(const FArguments& InArgs);
+
+	/** Fetches activities from the server and updates the list view. */
+	void ReloadActivities(TMap<FGuid, FConcertClientInfo> EndpointClientInfoMap, TArray<FConcertSessionActivity> FetchedActivities);
+	
+	/** Callback for handling the a new or updated activity item. */ 
+	void HandleActivityAddedOrUpdated(const FConcertClientInfo& InClientInfo, const FConcertSyncActivity& InActivity, const FStructOnScope& InActivitySummary);
+
+private:
+
+	/** Invoked when the text in the search box widget changes. */
+	void OnSearchTextChanged(const FText& InSearchText);
+
+	/** Invoked when the text in the search box widget is committed. */
+	void OnSearchTextCommitted(const FText& InFilterText, ETextCommit::Type CommitType);
+
+	/** Returns the text to highlight when the search bar has a text set. */
+	FText HighlightSearchedText() const;
+
+private:
+	
+	/** Holds the map of endpoint IDs to client info. */
+	TMap<FGuid, FConcertClientInfo> EndpointClientInfoMap;
+
+	/** Holds the map of activity IDs to Concert activities. */
+	TMap<int64, TSharedPtr<FConcertSessionActivity>> ActivityMap;
+
+	/** Display the activity list. */
+	TSharedPtr<SConcertSessionActivities> ActivityListView;
+
+	/** Controls the activity list view options */
+	TSharedPtr<FConcertSessionActivitiesOptions> ActivityListViewOptions;
+
+	/** The widget used to enter the text to search. */
+	TSharedPtr<SSearchBox> SearchBox;
+
+	/** The searched text to highlight. */
+	FText SearchedText;
+
+	/** Used to limit activities to a given package only. */
+	FName PackageNameFilter;
+
+	TOptional<FConcertClientInfo> GetClientInfo(FGuid Guid) const;
+};
\ No newline at end of file
diff --git a/Engine/Plugins/Developer/Concert/ConcertUI/ConcertSharedSlate/Source/ConcertSharedSlate/Public/Session/History/SSessionHistoryWrapper.h b/Engine/Plugins/Developer/Concert/ConcertUI/ConcertSharedSlate/Source/ConcertSharedSlate/Public/Session/History/SSessionHistoryWrapper.h
new file mode 100644
--- /dev/null
+++ b/Engine/Plugins/Developer/Concert/ConcertUI/ConcertSharedSlate/Source/ConcertSharedSlate/Public/Session/History/SSessionHistoryWrapper.h
@@ -0,0 +1,25 @@
+﻿// Copyright Epic Games, Inc. All Rights Reserved.
+
+#pragma once
+
+#include "CoreMinimal.h"
+#include "AbstractSessionHistoryController.h"
+#include "SSessionHistory.h"
+#include "Widgets/SCompoundWidget.h"
+#include "Widgets/DeclarativeSyntaxSupport.h"
+
+/** Shows a SSessionHistory widget - keeping alive the controller as long as this widget lives. Useful for closeable tabs. */
+class CONCERTSHAREDSLATE_API SSessionHistoryWrapper : public SCompoundWidget
+{
+public:
+
+	SLATE_BEGIN_ARGS(SSessionHistoryWrapper) {}
+	SLATE_END_ARGS()
+
+	void Construct(const FArguments& InArgs, TSharedRef<FAbstractSessionHistoryController> InController);
+
+private:
+
+	/** Manages SessionHistory */
+	TSharedPtr<FAbstractSessionHistoryController> Controller;
+};
diff --git a/Engine/Plugins/Developer/Concert/ConcertUI/ConcertSharedSlate/Source/ConcertSharedSlate/Public/SessionBrowser/ConcertBrowserUtils.h b/Engine/Plugins/Developer/Concert/ConcertUI/ConcertSharedSlate/Source/ConcertSharedSlate/Public/SessionBrowser/ConcertBrowserUtils.h
deleted file mode 100644
diff --git a/Engine/Plugins/Developer/Concert/ConcertUI/ConcertSharedSlate/Source/ConcertSharedSlate/Public/SessionBrowser/ConcertSessionBrowserSettings.h b/Engine/Plugins/Developer/Concert/ConcertUI/ConcertSharedSlate/Source/ConcertSharedSlate/Public/SessionBrowser/ConcertSessionBrowserSettings.h
deleted file mode 100644
diff --git a/Engine/Plugins/Developer/Concert/ConcertUI/ConcertSharedSlate/Source/ConcertSharedSlate/Public/SessionBrowser/ConcertSessionItem.h b/Engine/Plugins/Developer/Concert/ConcertUI/ConcertSharedSlate/Source/ConcertSharedSlate/Public/SessionBrowser/ConcertSessionItem.h
deleted file mode 100644
diff --git a/Engine/Plugins/Developer/Concert/ConcertUI/ConcertSharedSlate/Source/ConcertSharedSlate/Public/SessionBrowser/IConcertSessionBrowserController.h b/Engine/Plugins/Developer/Concert/ConcertUI/ConcertSharedSlate/Source/ConcertSharedSlate/Public/SessionBrowser/IConcertSessionBrowserController.h
deleted file mode 100644
diff --git a/Engine/Plugins/Developer/Concert/ConcertUI/ConcertSharedSlate/Source/ConcertSharedSlate/Public/SessionBrowser/SConcertSessionBrowser.h b/Engine/Plugins/Developer/Concert/ConcertUI/ConcertSharedSlate/Source/ConcertSharedSlate/Public/SessionBrowser/SConcertSessionBrowser.h
deleted file mode 100644
diff --git a/Engine/Plugins/Developer/Concert/ConcertUI/ConcertSharedSlate/Source/ConcertSharedSlate/Public/SessionHistory/AbstractSessionHistoryController.h b/Engine/Plugins/Developer/Concert/ConcertUI/ConcertSharedSlate/Source/ConcertSharedSlate/Public/SessionHistory/AbstractSessionHistoryController.h
deleted file mode 100644
diff --git a/Engine/Plugins/Developer/Concert/ConcertUI/ConcertSharedSlate/Source/ConcertSharedSlate/Public/SessionHistory/SSessionHistory.h b/Engine/Plugins/Developer/Concert/ConcertUI/ConcertSharedSlate/Source/ConcertSharedSlate/Public/SessionHistory/SSessionHistory.h
deleted file mode 100644
diff --git a/Engine/Plugins/Developer/Concert/ConcertUI/ConcertSharedSlate/Source/ConcertSharedSlate/Public/SessionHistory/SSessionHistoryWrapper.h b/Engine/Plugins/Developer/Concert/ConcertUI/ConcertSharedSlate/Source/ConcertSharedSlate/Public/SessionHistory/SSessionHistoryWrapper.h
deleted file mode 100644
diff --git a/Engine/Plugins/Developer/Concert/MultiUserClient/Source/MultiUserClient/Private/ConcertWorkspaceUI.cpp b/Engine/Plugins/Developer/Concert/MultiUserClient/Source/MultiUserClient/Private/ConcertWorkspaceUI.cpp
--- a/Engine/Plugins/Developer/Concert/MultiUserClient/Source/MultiUserClient/Private/ConcertWorkspaceUI.cpp
+++ b/Engine/Plugins/Developer/Concert/MultiUserClient/Source/MultiUserClient/Private/ConcertWorkspaceUI.cpp
@@ -31,15 +31,13 @@
 #include "Framework/MultiBox/MultiBoxExtender.h"
 #include "Framework/Notifications/NotificationManager.h"
 
-#include "SessionHistory/SSessionHistoryWrapper.h"
-#include "ToolMenuContext.h"
+#include "Session/History/SSessionHistoryWrapper.h"
 #include "ToolMenus.h"
 #include "Widgets/ClientSessionHistoryController.h"
 #include "Widgets/DeclarativeSyntaxSupport.h"
 #include "Widgets/SConcertSandboxPersistWidget.h"
 #include "Widgets/SCompoundWidget.h"
 #include "Widgets/SBoxPanel.h"
-#include "Widgets/Layout/SBox.h"
 #include "Widgets/Layout/SBorder.h"
 #include "Widgets/Images/SImage.h"
 #include "Widgets/Text/STextBlock.h"
diff --git a/Engine/Plugins/Developer/Concert/MultiUserClient/Source/MultiUserClient/Private/Widgets/ClientSessionHistoryController.cpp b/Engine/Plugins/Developer/Concert/MultiUserClient/Source/MultiUserClient/Private/Widgets/ClientSessionHistoryController.cpp
--- a/Engine/Plugins/Developer/Concert/MultiUserClient/Source/MultiUserClient/Private/Widgets/ClientSessionHistoryController.cpp
+++ b/Engine/Plugins/Developer/Concert/MultiUserClient/Source/MultiUserClient/Private/Widgets/ClientSessionHistoryController.cpp
@@ -6,7 +6,7 @@
 #include "IConcertClientWorkspace.h"
 #include "IConcertSyncClient.h"
 
-#include "SessionHistory/SSessionHistory.h"
+#include "Session/History/SSessionHistory.h"
 
 FClientSessionHistoryController::FClientSessionHistoryController(TSharedRef<IConcertSyncClient> Client, FName PackageFilter)
 	: FAbstractSessionHistoryController(PackageFilter)
diff --git a/Engine/Plugins/Developer/Concert/MultiUserClient/Source/MultiUserClient/Private/Widgets/ClientSessionHistoryController.h b/Engine/Plugins/Developer/Concert/MultiUserClient/Source/MultiUserClient/Private/Widgets/ClientSessionHistoryController.h
--- a/Engine/Plugins/Developer/Concert/MultiUserClient/Source/MultiUserClient/Private/Widgets/ClientSessionHistoryController.h
+++ b/Engine/Plugins/Developer/Concert/MultiUserClient/Source/MultiUserClient/Private/Widgets/ClientSessionHistoryController.h
@@ -3,7 +3,7 @@
 #pragma once
 
 #include "CoreMinimal.h"
-#include "SessionHistory/AbstractSessionHistoryController.h"
+#include "Session/History/AbstractSessionHistoryController.h"
 
 struct FConcertSyncActivity;
 struct FConcertClientInfo;
diff --git a/Engine/Plugins/Developer/Concert/MultiUserClient/Source/MultiUserClient/Private/Widgets/SActiveSession.cpp b/Engine/Plugins/Developer/Concert/MultiUserClient/Source/MultiUserClient/Private/Widgets/SActiveSession.cpp
--- a/Engine/Plugins/Developer/Concert/MultiUserClient/Source/MultiUserClient/Private/Widgets/SActiveSession.cpp
+++ b/Engine/Plugins/Developer/Concert/MultiUserClient/Source/MultiUserClient/Private/Widgets/SActiveSession.cpp
@@ -17,8 +17,8 @@
 #include "FileHelpers.h"
 #include "Misc/PackageName.h"
 #include "Styling/AppStyle.h"
-#include "SessionHistory/SSessionHistory.h"
-#include "SessionHistory/SSessionHistoryWrapper.h"
+#include "Session/History/SSessionHistory.h"
+#include "Session/History/SSessionHistoryWrapper.h"
 #include "Widgets/Text/STextBlock.h"
 #include "Widgets/Images/SImage.h"
 #include "Widgets/Input/SHyperlink.h"
diff --git a/Engine/Plugins/Developer/Concert/MultiUserClient/Source/MultiUserClient/Private/Widgets/SConcertBrowser.cpp b/Engine/Plugins/Developer/Concert/MultiUserClient/Source/MultiUserClient/Private/Widgets/SConcertBrowser.cpp
--- a/Engine/Plugins/Developer/Concert/MultiUserClient/Source/MultiUserClient/Private/Widgets/SConcertBrowser.cpp
+++ b/Engine/Plugins/Developer/Concert/MultiUserClient/Source/MultiUserClient/Private/Widgets/SConcertBrowser.cpp
@@ -14,10 +14,10 @@
 #include "MultiUserClientUtils.h"
 #include "SActiveSession.h"
 #include "SConcertSessionRecovery.h"
-#include "SessionBrowser/ConcertBrowserUtils.h"
-#include "SessionBrowser/ConcertSessionItem.h"
-#include "SessionBrowser/IConcertSessionBrowserController.h"
-#include "SessionBrowser/SConcertSessionBrowser.h"
+#include "Session/Browser/ConcertBrowserUtils.h"
+#include "Session/Browser/ConcertSessionItem.h"
+#include "Session/Browser/IConcertSessionBrowserController.h"
+#include "Session/Browser/SConcertSessionBrowser.h"
 
 #include "Algo/Transform.h"
 #include "Framework/Application/SlateApplication.h"
