Refactor column creation for SConcertSessionActivities

diff --git a/Engine/Plugins/Developer/Concert/ConcertUI/ConcertServerUI/Source/ConcertServerUI/Private/Widgets/Session/PackageViewer/SConcertSessionPackageViewer.cpp b/Engine/Plugins/Developer/Concert/ConcertUI/ConcertServerUI/Source/ConcertServerUI/Private/Widgets/Session/PackageViewer/SConcertSessionPackageViewer.cpp
--- a/Engine/Plugins/Developer/Concert/ConcertUI/ConcertServerUI/Source/ConcertServerUI/Private/Widgets/Session/PackageViewer/SConcertSessionPackageViewer.cpp
+++ b/Engine/Plugins/Developer/Concert/ConcertUI/ConcertServerUI/Source/ConcertServerUI/Private/Widgets/Session/PackageViewer/SConcertSessionPackageViewer.cpp
@@ -2,6 +2,7 @@
 
 #include "SConcertSessionPackageViewer.h"
 
+#include "Session/Activity/PredefinedActivityColumns.h"
 #include "Session/Activity/SConcertSessionActivities.h"
 #include "Widgets/Input/SSearchBox.h"
 
@@ -9,22 +10,23 @@
 
 void SConcertSessionPackageViewer::Construct(const FArguments& InArgs)
 {
+	using namespace UE::ConcertSharedSlate;
+	
 	ActivityListViewOptions = MakeShared<FConcertSessionActivitiesOptions>();
 	ActivityListViewOptions->bEnableConnectionActivityFiltering = false;
 	ActivityListViewOptions->bEnableLockActivityFiltering = false;
 	ActivityListViewOptions->bEnablePackageActivityFiltering = false;
 	ActivityListViewOptions->bEnableTransactionActivityFiltering = false;
-	
+
 	SAssignNew(ActivityListView, SConcertSessionActivities)
 		.OnGetPackageEvent(InArgs._GetPackageEvent)
 		.OnFetchActivities(InArgs._FetchInitialActivities)
 		.OnMapActivityToClient(InArgs._GetClientInfo)
 		.HighlightText(this, &SConcertSessionPackageViewer::HighlightSearchedText)
+		.Columns({
+			// TODO: Create Icon (depending on EConcertPackageUpdateType), Size, and Version
+		})
 		.TimeFormat(ActivityListViewOptions.Get(), &FConcertSessionActivitiesOptions::GetTimeFormat)
-		.ClientNameColumnVisibility(EVisibility::Visible)
-		.ClientAvatarColorColumnVisibility(EVisibility::Collapsed)
-		.OperationColumnVisibility(EVisibility::Collapsed)
-		.PackageColumnVisibility(EVisibility::Collapsed)
 		.ConnectionActivitiesVisibility(EVisibility::Collapsed)
 		.LockActivitiesVisibility(EVisibility::Collapsed)
 		.PackageActivitiesVisibility(EVisibility::Visible)
diff --git a/Engine/Plugins/Developer/Concert/ConcertUI/ConcertSharedSlate/Source/ConcertSharedSlate/Private/SConcertSessionRecovery.cpp b/Engine/Plugins/Developer/Concert/ConcertUI/ConcertSharedSlate/Source/ConcertSharedSlate/Private/SConcertSessionRecovery.cpp
--- a/Engine/Plugins/Developer/Concert/ConcertUI/ConcertSharedSlate/Source/ConcertSharedSlate/Private/SConcertSessionRecovery.cpp
+++ b/Engine/Plugins/Developer/Concert/ConcertUI/ConcertSharedSlate/Source/ConcertSharedSlate/Private/SConcertSessionRecovery.cpp
@@ -13,6 +13,8 @@
 #include "SNegativeActionButton.h"
 #include "SPositiveActionButton.h"
 
+#include "Session/Activity/PredefinedActivityColumns.h"
+
 #define LOCTEXT_NAMESPACE "SConcertSessionRecovery"
 
 void SConcertSessionRecovery::Construct(const FArguments& InArgs)
@@ -29,24 +31,21 @@ void SConcertSessionRecovery::Construct(const FArguments& InArgs)
 	ActivityViewOptions->bEnablePackageActivityFiltering = InArgs._IsPackageActivityFilteringEnabled;
 	ActivityViewOptions->bEnableTransactionActivityFiltering = InArgs._IsTransactionActivityFilteringEnabled;
 	ActivityViewOptions->bEnableIgnoredActivityFiltering = InArgs._IsIgnoredActivityFilteringEnabled;
-
+	
 	SAssignNew(ActivityView, SConcertSessionActivities)
-	.OnFetchActivities(InArgs._OnFetchActivities)
-	.OnMapActivityToClient(InArgs._OnMapActivityToClient)
-	.OnMakeColumnOverlayWidget_Lambda([this](TWeakPtr<FConcertSessionActivity> Activity, const FName& ColumnId) { return MakeRecoverThroughWidget(Activity, ColumnId); })
-	.HighlightText(this, &SConcertSessionRecovery::HighlightSearchText)
-	.TimeFormat(ActivityViewOptions.Get(), &FConcertSessionActivitiesOptions::GetTimeFormat)
-	.ClientAvatarColorColumnVisibility(InArgs._ClientAvatarColorColumnVisibility)
-	.ClientNameColumnVisibility(InArgs._ClientNameColumnVisibility)
-	.OperationColumnVisibility(InArgs._OperationColumnVisibility)
-	.PackageColumnVisibility(InArgs._PackageColumnVisibility)
-	.ConnectionActivitiesVisibility(ActivityViewOptions.Get(), &FConcertSessionActivitiesOptions::GetConnectionActivitiesVisibility)
-	.LockActivitiesVisibility(ActivityViewOptions.Get(), &FConcertSessionActivitiesOptions::GetLockActivitiesVisibility)
-	.PackageActivitiesVisibility(ActivityViewOptions.Get(), &FConcertSessionActivitiesOptions::GetPackageActivitiesVisibility)
-	.TransactionActivitiesVisibility(ActivityViewOptions.Get(), &FConcertSessionActivitiesOptions::GetTransactionActivitiesVisibility)
-	.IgnoredActivitiesVisibility(ActivityViewOptions.Get(), &FConcertSessionActivitiesOptions::GetIgnoredActivitiesVisibility)
-	.DetailsAreaVisibility(InArgs._DetailsAreaVisibility)
-	.NoActivitiesReasonText(InArgs._NoActivitiesReasonText);
+		.OnFetchActivities(InArgs._OnFetchActivities)
+		.OnMapActivityToClient(InArgs._OnMapActivityToClient)
+		.OnMakeColumnOverlayWidget_Lambda([this](TWeakPtr<FConcertSessionActivity> Activity, const FName& ColumnId) { return MakeRecoverThroughWidget(Activity, ColumnId); })
+		.Columns(MakeColumns(InArgs))
+		.HighlightText(this, &SConcertSessionRecovery::HighlightSearchText)
+		.TimeFormat(ActivityViewOptions.Get(), &FConcertSessionActivitiesOptions::GetTimeFormat)
+		.ConnectionActivitiesVisibility(ActivityViewOptions.Get(), &FConcertSessionActivitiesOptions::GetConnectionActivitiesVisibility)
+		.LockActivitiesVisibility(ActivityViewOptions.Get(), &FConcertSessionActivitiesOptions::GetLockActivitiesVisibility)
+		.PackageActivitiesVisibility(ActivityViewOptions.Get(), &FConcertSessionActivitiesOptions::GetPackageActivitiesVisibility)
+		.TransactionActivitiesVisibility(ActivityViewOptions.Get(), &FConcertSessionActivitiesOptions::GetTransactionActivitiesVisibility)
+		.IgnoredActivitiesVisibility(ActivityViewOptions.Get(), &FConcertSessionActivitiesOptions::GetIgnoredActivitiesVisibility)
+		.DetailsAreaVisibility(InArgs._DetailsAreaVisibility)
+		.NoActivitiesReasonText(InArgs._NoActivitiesReasonText);
 
 	ChildSlot
 	[
@@ -143,6 +142,31 @@ void SConcertSessionRecovery::Construct(const FArguments& InArgs)
 	];
 }
 
+TArray<FActivityColumn> SConcertSessionRecovery::MakeColumns(const FArguments& InArgs) const
+{
+	using namespace UE::ConcertSharedSlate;
+	TArray<FActivityColumn> Result;
+
+	if (InArgs._WithClientAvatarColorColumn)
+	{
+		Result.Add(ActivityColumn::AvatarColor());
+	}
+	if (InArgs._WithClientNameColumn)
+	{
+		Result.Add(ActivityColumn::ClientName());
+	}
+	if (InArgs._WithOperationColumn)
+	{
+		Result.Add(ActivityColumn::Operation());
+	}
+	if (InArgs._WithPackageColumn)
+	{
+		Result.Add(ActivityColumn::Package());
+	}
+	
+	return Result;
+}
+
 TSharedPtr<SWidget> SConcertSessionRecovery::MakeRecoverThroughWidget(TWeakPtr<FConcertSessionActivity> Activity, const FName& ColumnId)
 {
 	if (TSharedPtr<FConcertSessionActivity> ActivityPin = Activity.Pin())
diff --git a/Engine/Plugins/Developer/Concert/ConcertUI/ConcertSharedSlate/Source/ConcertSharedSlate/Private/Session/Activity/ActivityColumn.cpp b/Engine/Plugins/Developer/Concert/ConcertUI/ConcertSharedSlate/Source/ConcertSharedSlate/Private/Session/Activity/ActivityColumn.cpp
new file mode 100644
--- /dev/null
+++ b/Engine/Plugins/Developer/Concert/ConcertUI/ConcertSharedSlate/Source/ConcertSharedSlate/Private/Session/Activity/ActivityColumn.cpp
@@ -0,0 +1,19 @@
+﻿// Copyright Epic Games, Inc. All Rights Reserved.
+
+#include "Session/Activity/ActivityColumn.h"
+
+void FActivityColumn::BuildColumnWidget(const TSharedRef<SConcertSessionActivities>& Owner,  const TSharedRef<FConcertSessionActivity>& Activity, SOverlay::FScopedWidgetSlotArguments& Slot) const
+{
+    if (GenerateColumnWidgetCallback.IsBound())
+    {
+	    GenerateColumnWidgetCallback.Execute(Owner, Activity, Slot);
+    }
+}
+
+void FActivityColumn::ExecutePopulateSearchString(const TSharedRef<SConcertSessionActivities>& Owner, const FConcertSessionActivity& Activity, TArray<FString>& ExistingSearchStrings) const
+{
+	if (PopulateSearchStringCallback.IsBound())
+	{
+		PopulateSearchStringCallback.Execute(Owner, Activity, ExistingSearchStrings);
+	}
+}
diff --git a/Engine/Plugins/Developer/Concert/ConcertUI/ConcertSharedSlate/Source/ConcertSharedSlate/Private/Session/Activity/PredefinedActivityColumns.cpp b/Engine/Plugins/Developer/Concert/ConcertUI/ConcertSharedSlate/Source/ConcertSharedSlate/Private/Session/Activity/PredefinedActivityColumns.cpp
new file mode 100644
--- /dev/null
+++ b/Engine/Plugins/Developer/Concert/ConcertUI/ConcertSharedSlate/Source/ConcertSharedSlate/Private/Session/Activity/PredefinedActivityColumns.cpp
@@ -0,0 +1,203 @@
+﻿// Copyright Epic Games, Inc. All Rights Reserved.
+
+#include "Session/Activity/PredefinedActivityColumns.h"
+
+#include "ConcertFrontendStyle.h"
+#include "ConcertFrontendUtils.h"
+#include "ConcertSyncSessionTypes.h"
+#include "ConcertWorkspaceData.h"
+#include "SessionActivityUtils.h"
+
+#include "Session/Activity/ActivityColumn.h"
+#include "Widgets/SOverlay.h"
+#include "Widgets/SUndoHistoryDetails.h"
+#include "Widgets/Colors/SColorBlock.h"
+#include "Widgets/Text/SRichTextBlock.h"
+#include "Widgets/Text/STextBlock.h"
+
+#define LOCTEXT_NAMESPACE "SConcertSessionActivities"
+
+namespace UE::ConcertSharedSlate::Private
+{
+	static FText GetActivityDateTime(const FConcertSessionActivity& Activity, SConcertSessionActivities::ETimeFormat TimeFormat)
+	{
+		return TimeFormat == SConcertSessionActivities::ETimeFormat::Relative ? ConcertFrontendUtils::FormatRelativeTime(Activity.Activity.EventTime) : FText::AsDateTime(Activity.Activity.EventTime);
+	}
+}
+
+const FName UE::ConcertSharedSlate::ActivityColumn::DateTimeColumnId	= TEXT("DateTime");
+const FName UE::ConcertSharedSlate::ActivityColumn::SummaryColumnId     = TEXT("Summary");
+const FName UE::ConcertSharedSlate::ActivityColumn::AvatarColorColumnId = TEXT("Client AvatarColor");
+const FName UE::ConcertSharedSlate::ActivityColumn::ClientNameColumnId  = TEXT("Client");
+const FName UE::ConcertSharedSlate::ActivityColumn::OperationColumnId   = TEXT("Operation");
+const FName UE::ConcertSharedSlate::ActivityColumn::PackageColumnId     = TEXT("Package");
+
+FActivityColumn UE::ConcertSharedSlate::ActivityColumn::DateTime()
+{
+	return FActivityColumn(
+		SHeaderRow::Column(DateTimeColumnId)
+			.DefaultLabel(LOCTEXT("DateTime", "Date/Time"))
+			.ManualWidth(160)
+		)
+		.ColumnSortOrder(static_cast<int32>(EPredefinedColumnOrder::DateTime))
+		.GenerateColumnWidget(FActivityColumn::FGenerateColumnWidget::CreateLambda([](const TSharedRef<SConcertSessionActivities>& Owner, const TSharedRef<FConcertSessionActivity>& Activity, SOverlay::FScopedWidgetSlotArguments& Slot)
+		{
+			Slot
+			.VAlign(VAlign_Center)
+			[
+				SNew(STextBlock)
+				.Text_Lambda([TimeFormat = Owner->GetTimeFormat(), WeakActivity = TWeakPtr<FConcertSessionActivity>(Activity)]()
+				{
+					if (const TSharedPtr<FConcertSessionActivity> ItemPin = WeakActivity.Pin())
+					{
+						return Private::GetActivityDateTime(*ItemPin, TimeFormat.Get());
+					}
+					return FText::GetEmpty();
+				})
+				.HighlightText(Owner->GetHighlightText())
+			];
+		}))
+		.PopulateSearchString(FActivityColumn::FPopulateSearchString::CreateLambda([](const TSharedRef<SConcertSessionActivities>& Owner, const FConcertSessionActivity& Activity, TArray<FString>& OutSearchStrings)
+		{
+			OutSearchStrings.Add(Private::GetActivityDateTime(Activity, Owner->GetTimeFormat().Get()).ToString());
+		}));
+}
+
+FActivityColumn UE::ConcertSharedSlate::ActivityColumn::Summary()
+{
+	return FActivityColumn(
+		SHeaderRow::Column(SummaryColumnId)
+			.DefaultLabel(LOCTEXT("Summary", "Summary"))
+		)
+		.GenerateColumnWidget(FActivityColumn::FGenerateColumnWidget::CreateLambda([](const TSharedRef<SConcertSessionActivities>& Owner, const TSharedRef<FConcertSessionActivity>& Activity, SOverlay::FScopedWidgetSlotArguments& Slot)
+		{
+			const FConcertSyncActivitySummary* Summary = Activity->ActivitySummary.Cast<FConcertSyncActivitySummary>();
+			constexpr bool bRichText = true;
+			const FText SummaryText = Summary ? Summary->ToDisplayText(FText::GetEmpty(), bRichText) : FText::GetEmpty();
+			Slot
+			.VAlign(VAlign_Center)
+			[
+				SNew(SRichTextBlock)
+				.DecoratorStyleSet(FConcertFrontendStyle::Get().Get())
+				.Text(SummaryText)
+				.HighlightText(Owner->GetHighlightText())
+			];
+		}))
+		.PopulateSearchString(FActivityColumn::FPopulateSearchString::CreateLambda([](const TSharedRef<SConcertSessionActivities>& Owner, const FConcertSessionActivity& Activity, TArray<FString>& OutSearchStrings)
+		{
+			if (const FConcertSyncActivitySummary* Summary = Activity.ActivitySummary.Cast<FConcertSyncActivitySummary>();
+				Summary)
+			{
+				const TOptional<FConcertClientInfo> ClientInfo = Owner->GetClientInfo(Activity.Activity.EndpointId);
+				const FText ClientName = ClientInfo ? FText::AsCultureInvariant(ClientInfo->DisplayName) : FText::GetEmpty();
+
+				constexpr bool bRichText = false;
+				OutSearchStrings.Add(Summary->ToDisplayText(ClientName, bRichText).ToString());
+			}
+		}));
+}
+
+FActivityColumn UE::ConcertSharedSlate::ActivityColumn::AvatarColor()
+{
+	return FActivityColumn(
+		SHeaderRow::Column(AvatarColorColumnId)
+			.DefaultLabel(INVTEXT(""))
+			.ManualWidth(8)
+		)
+		.ColumnSortOrder(static_cast<int32>(EPredefinedColumnOrder::AvatarColor))
+		.GenerateColumnWidget(FActivityColumn::FGenerateColumnWidget::CreateLambda([](const TSharedRef<SConcertSessionActivities>& Owner, const TSharedRef<FConcertSessionActivity>& Activity, SOverlay::FScopedWidgetSlotArguments& Slot)
+		{
+			const TOptional<FConcertClientInfo> ClientInfo = Owner->GetClientInfo(Activity->Activity.EndpointId);
+			const FLinearColor AvatarColor = ClientInfo ? ClientInfo->AvatarColor : FConcertFrontendStyle::Get()->GetColor("Concert.DisconnectedColor");
+			Slot
+			.HAlign(HAlign_Center)
+			.VAlign(VAlign_Center)
+			.Padding(2, 1)
+			[
+				SNew(SColorBlock)
+				.Color(AvatarColor)
+				.Size(FVector2D(4.f, 16.f))
+			];
+		}));
+}
+
+FActivityColumn UE::ConcertSharedSlate::ActivityColumn::ClientName()
+{
+	return FActivityColumn(
+		SHeaderRow::Column(ClientNameColumnId)
+			.DefaultLabel(LOCTEXT("Client", "Client"))
+			.ManualWidth(80)
+		)
+		.ColumnSortOrder(static_cast<int32>(EPredefinedColumnOrder::ClientName))
+		.GenerateColumnWidget(FActivityColumn::FGenerateColumnWidget::CreateLambda([](const TSharedRef<SConcertSessionActivities>& Owner, const TSharedRef<FConcertSessionActivity>& Activity, SOverlay::FScopedWidgetSlotArguments& Slot)
+		{
+			const TOptional<FConcertClientInfo> ClientInfo = Owner->GetClientInfo(Activity.Get().Activity.EndpointId);
+			const FText ClientName = ClientInfo ? FText::AsCultureInvariant(ClientInfo->DisplayName) : FText::GetEmpty();
+			
+			Slot
+			.VAlign(VAlign_Center)
+			[
+				SNew(STextBlock)
+				.Text(ClientName)
+				.HighlightText(Owner->GetHighlightText())
+			];
+		}))
+		.PopulateSearchString(FActivityColumn::FPopulateSearchString::CreateLambda([](const TSharedRef<SConcertSessionActivities>& Owner, const FConcertSessionActivity& Activity, TArray<FString>& OutSearchStrings)
+		{
+			if (const TOptional<FConcertClientInfo> ClientInfo = Owner->GetClientInfo(Activity.Activity.EndpointId))
+			{
+				const FText ClientName = ClientInfo ? FText::AsCultureInvariant(ClientInfo->DisplayName) : FText::GetEmpty();
+				OutSearchStrings.Add(ClientName.ToString());
+			}
+		}));
+}
+
+FActivityColumn UE::ConcertSharedSlate::ActivityColumn::Operation()
+{
+	return FActivityColumn(
+		SHeaderRow::Column(OperationColumnId)
+			.DefaultLabel(LOCTEXT("Operation", "Operation"))
+			.ManualWidth(160)
+		)
+		.ColumnSortOrder(static_cast<int32>(EPredefinedColumnOrder::Operation))
+		.GenerateColumnWidget(FActivityColumn::FGenerateColumnWidget::CreateLambda([](const TSharedRef<SConcertSessionActivities>& Owner, const TSharedRef<FConcertSessionActivity>& Activity, SOverlay::FScopedWidgetSlotArguments& Slot)
+		{
+			Slot
+			.VAlign(VAlign_Center)
+			[
+				SNew(STextBlock)
+				.Text(Private::GetOperationName(*Activity))
+				.HighlightText(Owner->GetHighlightText())
+			];
+		}))
+		.PopulateSearchString(FActivityColumn::FPopulateSearchString::CreateLambda([](const TSharedRef<SConcertSessionActivities>& Owner, const FConcertSessionActivity& Activity, TArray<FString>& OutSearchStrings)
+		{
+			OutSearchStrings.Add(Private::GetOperationName(Activity).ToString());
+		}));
+}
+
+FActivityColumn UE::ConcertSharedSlate::ActivityColumn::Package()
+{
+	return FActivityColumn(
+		SHeaderRow::Column(PackageColumnId)
+			.DefaultLabel(LOCTEXT("Package", "Package"))
+			.ManualWidth(200)
+		)
+		.ColumnSortOrder(static_cast<int32>(EPredefinedColumnOrder::Package))
+		.GenerateColumnWidget(FActivityColumn::FGenerateColumnWidget::CreateLambda([](const TSharedRef<SConcertSessionActivities>& Owner, const TSharedRef<FConcertSessionActivity>& Activity, SOverlay::FScopedWidgetSlotArguments& Slot)
+		{
+			Slot
+			.VAlign(VAlign_Center)
+			[
+				SNew(STextBlock)
+				.Text(Private::GetPackageName(*Activity))
+				.HighlightText(Owner->GetHighlightText())
+			];
+		}))
+		.PopulateSearchString(FActivityColumn::FPopulateSearchString::CreateLambda([](const TSharedRef<SConcertSessionActivities>& Owner, const FConcertSessionActivity& Activity, TArray<FString>& OutSearchStrings)
+		{
+			OutSearchStrings.Add(Private::GetPackageName(Activity).ToString());
+		}));
+}
+
+#undef LOCTEXT_NAMESPACE
\ No newline at end of file
diff --git a/Engine/Plugins/Developer/Concert/ConcertUI/ConcertSharedSlate/Source/ConcertSharedSlate/Private/Session/Activity/SConcertSessionActivities.cpp b/Engine/Plugins/Developer/Concert/ConcertUI/ConcertSharedSlate/Source/ConcertSharedSlate/Private/Session/Activity/SConcertSessionActivities.cpp
--- a/Engine/Plugins/Developer/Concert/ConcertUI/ConcertSharedSlate/Source/ConcertSharedSlate/Private/Session/Activity/SConcertSessionActivities.cpp
+++ b/Engine/Plugins/Developer/Concert/ConcertUI/ConcertSharedSlate/Source/ConcertSharedSlate/Private/Session/Activity/SConcertSessionActivities.cpp
@@ -27,20 +27,15 @@
 #include "ConcertLogGlobal.h"
 #include "ConcertTransactionEvents.h"
 #include "ConcertWorkspaceData.h"
+#include "SessionActivityUtils.h"
 #include "SPackageDetails.h"
 
+#include "Session/Activity/PredefinedActivityColumns.h"
+
 #define LOCTEXT_NAMESPACE "SConcertSessionActivities"
 
 namespace ConcertSessionActivityUtils
 {
-// The columns names.
-const FName DateTimeColumnId    = TEXT("DateTime");
-const FName OperationColumnId   = TEXT("Operation");
-const FName PackageColumnId     = TEXT("Package");
-const FName SummaryColumnId     = TEXT("Summary");
-const FName ClientNameColumnId  = TEXT("Client");
-const FName AvatarColorColumnId = TEXT("Client AvatarColor");
-
 // The View Options check boxes.
 const FName DisplayRelativeTimeCheckBoxId       = TEXT("DisplayRelativeTime");
 const FName ShowConnectionActivitiesCheckBoxId  = TEXT("ShowConnectionActivities");
@@ -54,73 +49,6 @@ FText GetActivityDateTime(const FConcertSessionActivity& Activity, SConcertSessi
 	return TimeFormat == SConcertSessionActivities::ETimeFormat::Relative ? ConcertFrontendUtils::FormatRelativeTime(Activity.Activity.EventTime) : FText::AsDateTime(Activity.Activity.EventTime);
 }
 
-FText GetOperationName(const FConcertSessionActivity& Activity)
-{
-	if (const FConcertSyncTransactionActivitySummary* TransactionSummary = Activity.ActivitySummary.Cast<FConcertSyncTransactionActivitySummary>())
-	{
-		return TransactionSummary->TransactionTitle;
-	}
-
-	if (const FConcertSyncPackageActivitySummary* PackageSummary = Activity.ActivitySummary.Cast<FConcertSyncPackageActivitySummary>())
-	{
-		auto GetSavedOperationNameFn = [](const FConcertSyncPackageActivitySummary* PackageSummary)
-		{
-			if (PackageSummary->bAutoSave)
-			{
-				return LOCTEXT("AutoSavePackageOperation", "Auto-Save Package");
-			}
-			return PackageSummary->bPreSave ? LOCTEXT("PreSavePackageOperation", "Pre-Save Package") : LOCTEXT("SavePackageOperation", "Save Package");
-		};
-
-		switch (PackageSummary->PackageUpdateType)
-		{
-		case EConcertPackageUpdateType::Added   : return LOCTEXT("NewPackageOperation",    "New Package");
-		case EConcertPackageUpdateType::Deleted : return LOCTEXT("DeletePackageOperation", "Delete Package");
-		case EConcertPackageUpdateType::Renamed : return LOCTEXT("RenamePackageOperation", "Rename Package");
-		case EConcertPackageUpdateType::Saved   : return GetSavedOperationNameFn(PackageSummary);
-		case EConcertPackageUpdateType::Dummy   : return LOCTEXT("DiscardPackageOperation", "Discard Changes");
-		default: break;
-		}
-	}
-
-	if (const FConcertSyncConnectionActivitySummary* ConnectionSummary = Activity.ActivitySummary.Cast<FConcertSyncConnectionActivitySummary>())
-	{
-		switch (ConnectionSummary->ConnectionEventType)
-		{
-		case EConcertSyncConnectionEventType::Connected:    return LOCTEXT("JoinOperation", "Join Session");
-		case EConcertSyncConnectionEventType::Disconnected: return LOCTEXT("LeaveOperation", "Leave Session");
-		default: break;
-		}
-	}
-
-	if (const FConcertSyncLockActivitySummary* LockSummary = Activity.ActivitySummary.Cast<FConcertSyncLockActivitySummary>())
-	{
-		switch (LockSummary->LockEventType)
-		{
-		case EConcertSyncLockEventType::Locked:   return LOCTEXT("LockOperation", "Lock");
-		case EConcertSyncLockEventType::Unlocked: return LOCTEXT("UnlockOperation", "Unlock");
-		default: break;
-		}
-	}
-
-	return FText::GetEmpty();
-}
-
-FText GetPackageName(const FConcertSessionActivity& Activity)
-{
-	if (const FConcertSyncPackageActivitySummary* Summary = Activity.ActivitySummary.Cast<FConcertSyncPackageActivitySummary>())
-	{
-		return FText::FromName(Summary->PackageName);
-	}
-
-	if (const FConcertSyncTransactionActivitySummary* Summary = Activity.ActivitySummary.Cast<FConcertSyncTransactionActivitySummary>())
-	{
-		return FText::FromName(Summary->PrimaryPackageName);
-	}
-
-	return FText::GetEmpty();
-}
-
 FText GetSummary(const FConcertSessionActivity& Activity, const FText& ClientName, bool bAsRichText)
 {
 	if (const FConcertSyncActivitySummary* Summary = Activity.ActivitySummary.Cast<FConcertSyncActivitySummary>())
@@ -149,39 +77,32 @@ FText GetClientName(const FConcertClientInfo* InActivityClient)
 class SConcertSessionActivityRow : public SMultiColumnTableRow<TSharedPtr<FConcertSessionActivity>>
 {
 public:
-	SLATE_BEGIN_ARGS(SConcertSessionActivityRow)
-		: _TimeFormat(SConcertSessionActivities::ETimeFormat::Relative)
-		, _HighlightText()
-		, _OnMakeColumnOverlayWidget()
-	{
-	}
-
-	SLATE_ATTRIBUTE(SConcertSessionActivities::ETimeFormat, TimeFormat)
-	SLATE_ATTRIBUTE(FText, HighlightText)
-	SLATE_ARGUMENT(SConcertSessionActivities::FMakeColumnOverlayWidgetFunc, OnMakeColumnOverlayWidget) // Function invoked when generating a row to add a widget above the column widget.
 
+	DECLARE_DELEGATE_RetVal_OneParam(const FActivityColumn*, FGetColumn,
+		const FName& ColumnId
+		);
+	
+	SLATE_BEGIN_ARGS(SConcertSessionActivityRow)
+		: _OnMakeColumnOverlayWidget()
+	{}
+		/** Function invoked when generating a row to add a widget above the column widget. */
+		SLATE_ARGUMENT(SConcertSessionActivities::FMakeColumnOverlayWidgetFunc, OnMakeColumnOverlayWidget) 
 	SLATE_END_ARGS()
 
-public:
 	/**
 	 * Constructs a row widget to display a Concert activity.
 	 * @param InArgs The widgets arguments.
-	 * @param InActivty The activity to display.
+	 * @param InActivity The activity to display.
 	 * @param InActivityClient The client who produced this activity. Can be null if unknown or not desirable.
 	 * @param InOwnerTableView The table view that will own this row.
 	 */
-	void Construct(const FArguments& InArgs, TSharedPtr<FConcertSessionActivity> InActivity, const TOptional<FConcertClientInfo>& InActivityClient, const TSharedRef<STableViewBase>& InOwnerTableView);
+	void Construct(const FArguments& InArgs, TSharedRef<SConcertSessionActivities> InOwner, TSharedRef<FConcertSessionActivity> InActivity, FGetColumn InColumnGetter, const TOptional<FConcertClientInfo>& InActivityClient, const TSharedRef<STableViewBase>& InOwnerTableView);
 
 	/** Generates the widget representing this row. */
 	virtual TSharedRef<SWidget> GenerateWidgetForColumn(const FName& ColumnName) override;
 
 private:
-	/** Format the the time displayed as relative or absolute according to the attribute 'TimeFormat'. */
-	FText FormatEventDateTime() const;
-
-	/** Returns the client avatar color. */
-	FLinearColor GetClientAvatarColor() const { return ClientAvatarColor; }
-
+	
 	/** Returns the display name of the client who performed the activity or an empty text if the information wasn't available or desired. */
 	FText GetClientName() const { return ClientName; }
 
@@ -189,25 +110,25 @@ private:
 	FText MakeTooltipText() const;
 
 private:
+	TWeakPtr<SConcertSessionActivities> Owner;
 	TWeakPtr<FConcertSessionActivity> Activity;
-	TAttribute<SConcertSessionActivities::ETimeFormat> TimeFormat;
+	FGetColumn ColumnGetter;
 	FText AbsoluteDateTime;
 	FText ClientName;
-	FLinearColor ClientAvatarColor;
-	TAttribute<FText> HighlightText;
 	SConcertSessionActivities::FMakeColumnOverlayWidgetFunc OnMakeColumnOverlayWidget;
 };
 
 
-void SConcertSessionActivityRow::Construct(const FArguments& InArgs, TSharedPtr<FConcertSessionActivity> InActivity, const TOptional<FConcertClientInfo>& InActivityClient, const TSharedRef<STableViewBase>& InOwnerTableView)
+void SConcertSessionActivityRow::Construct(const FArguments& InArgs, TSharedRef<SConcertSessionActivities> InOwner, TSharedRef<FConcertSessionActivity> InActivity, FGetColumn InColumnGetter, const TOptional<FConcertClientInfo>& InActivityClient, const TSharedRef<STableViewBase>& InOwnerTableView)
 {
+	check(InColumnGetter.IsBound());
+	Owner = InOwner;
 	Activity = InActivity;
-	TimeFormat = InArgs._TimeFormat;
-	HighlightText = InArgs._HighlightText;
+	ColumnGetter = InColumnGetter;
+	
 	OnMakeColumnOverlayWidget = InArgs._OnMakeColumnOverlayWidget;
 	AbsoluteDateTime = ConcertSessionActivityUtils::GetActivityDateTime(*InActivity, SConcertSessionActivities::ETimeFormat::Absolute); // Cache the absolute time. It doesn't changes.
 	ClientName = ConcertSessionActivityUtils::GetClientName(InActivityClient);
-	ClientAvatarColor = InActivityClient ? InActivityClient->AvatarColor : FConcertFrontendStyle::Get()->GetColor("Concert.DisconnectedColor");
 
 	// Construct base class
 	SMultiColumnTableRow<TSharedPtr<FConcertSessionActivity>>::Construct(FSuperRowType::FArguments(), InOwnerTableView);
@@ -220,75 +141,17 @@ void SConcertSessionActivityRow::Construct(const FArguments& InArgs, TSharedPtr<
 
 TSharedRef<SWidget> SConcertSessionActivityRow::GenerateWidgetForColumn(const FName& ColumnId)
 {
-	TSharedPtr<FConcertSessionActivity> ActivityPin = Activity.Pin();
-	TSharedRef<SOverlay> Overlay = SNew(SOverlay);
-
-	if (ColumnId == ConcertSessionActivityUtils::AvatarColorColumnId)
-	{
-		Overlay->AddSlot()
-		.HAlign(HAlign_Center)
-		.VAlign(VAlign_Center)
-		.Padding(2, 1)
-		[
-			SNew(SColorBlock)
-			.Color(GetClientAvatarColor())
-			.Size(FVector2D(4.f, 16.f))
-		];
-	}
-	else if (ColumnId == ConcertSessionActivityUtils::DateTimeColumnId)
-	{
-		Overlay->AddSlot()
-		.VAlign(VAlign_Center)
-		[
-			SNew(STextBlock)
-			.Text(this, &SConcertSessionActivityRow::FormatEventDateTime)
-			.HighlightText(HighlightText)
-		];
-	}
-	else if (ColumnId == ConcertSessionActivityUtils::ClientNameColumnId)
-	{
-		Overlay->AddSlot()
-		.VAlign(VAlign_Center)
-		[
-			SNew(STextBlock)
-			.Text(GetClientName())
-			.HighlightText(HighlightText)
-		];
-	}
-	else if (ColumnId == ConcertSessionActivityUtils::PackageColumnId)
-	{
-		Overlay->AddSlot()
-		.VAlign(VAlign_Center)
-		[
-			SNew(STextBlock)
-			.Text(ConcertSessionActivityUtils::GetPackageName(*ActivityPin))
-			.HighlightText(HighlightText)
-		];
-	}
-	else if (ColumnId == ConcertSessionActivityUtils::OperationColumnId)
-	{
-		Overlay->AddSlot()
-		.VAlign(VAlign_Center)
-		[
-			SNew(STextBlock)
-			.Text(ConcertSessionActivityUtils::GetOperationName(*ActivityPin))
-			.HighlightText(HighlightText)
-		];
-	}
-	else
+	const TSharedPtr<SConcertSessionActivities> OwnerPin = Owner.Pin();
+	const TSharedPtr<FConcertSessionActivity> ActivityPin = Activity.Pin();
+	check(OwnerPin && ActivityPin);
+	
+	const TSharedRef<SOverlay> Overlay = SNew(SOverlay);
+	if (const FActivityColumn* Column = ColumnGetter.Execute(ColumnId))
 	{
-		check(ColumnId == ConcertSessionActivityUtils::SummaryColumnId);
-
-		Overlay->AddSlot()
-		.VAlign(VAlign_Center)
-		[
-			SNew(SRichTextBlock)
-			.DecoratorStyleSet(FConcertFrontendStyle::Get().Get())
-			.Text(ConcertSessionActivityUtils::GetSummary(*ActivityPin, FText::GetEmpty(), true/*bAsRichText*/))
-			.HighlightText(HighlightText)
-		];
+		SOverlay::FScopedWidgetSlotArguments WidgetSlot = Overlay->AddSlot();
+		Column->BuildColumnWidget(OwnerPin.ToSharedRef(), ActivityPin.ToSharedRef(), WidgetSlot);
 	}
-
+	
 	if (OnMakeColumnOverlayWidget.IsBound())
 	{
 		if (TSharedPtr<SWidget> OverlayedWidget = OnMakeColumnOverlayWidget.Execute(ActivityPin, ColumnId))
@@ -301,7 +164,6 @@ TSharedRef<SWidget> SConcertSessionActivityRow::GenerateWidgetForColumn(const FN
 	}
 
 	SetToolTipText(TAttribute<FText>(this, &SConcertSessionActivityRow::MakeTooltipText));
-
 	return Overlay;
 }
 
@@ -309,10 +171,10 @@ FText SConcertSessionActivityRow::MakeTooltipText() const
 {
 	if (TSharedPtr<FConcertSessionActivity> ActivityPin = Activity.Pin())
 	{
-		FText Client = GetClientName();
-		FText Operation = ConcertSessionActivityUtils::GetOperationName(*ActivityPin);
-		FText Package = ConcertSessionActivityUtils::GetPackageName(*ActivityPin);
-		FText Summary = ConcertSessionActivityUtils::GetSummary(*ActivityPin, Client, /*bAsRichText*/false);
+		const FText Client = GetClientName();
+		const FText Operation = UE::ConcertSharedSlate::Private::GetOperationName(*ActivityPin);
+		const FText Package = UE::ConcertSharedSlate::Private::GetPackageName(*ActivityPin);
+		const FText Summary = ConcertSessionActivityUtils::GetSummary(*ActivityPin, Client, /*bAsRichText*/false);
 
 		FTextBuilder TextBuilder;
 
@@ -345,16 +207,6 @@ FText SConcertSessionActivityRow::MakeTooltipText() const
 	return FText::GetEmpty();
 }
 
-FText SConcertSessionActivityRow::FormatEventDateTime() const
-{
-	if (TSharedPtr<FConcertSessionActivity> ItemPin = Activity.Pin())
-	{
-		return TimeFormat.Get() == SConcertSessionActivities::ETimeFormat::Relative ? ConcertFrontendUtils::FormatRelativeTime(ItemPin->Activity.EventTime) : AbsoluteDateTime;
-	}
-	return FText::GetEmpty();
-}
-
-
 void SConcertSessionActivities::Construct(const FArguments& InArgs)
 {
 	FetchActivitiesFn = InArgs._OnFetchActivities;
@@ -362,11 +214,11 @@ void SConcertSessionActivities::Construct(const FArguments& InArgs)
 	GetTransactionEventFn = InArgs._OnGetTransactionEvent;
 	GetPackageEventFn = InArgs._OnGetPackageEvent;
 	MakeColumnOverlayWidgetFn = InArgs._OnMakeColumnOverlayWidget;
+	
+	Columns = InArgs._Columns;
 	HighlightText = InArgs._HighlightText;
+	
 	TimeFormat = InArgs._TimeFormat;
-	ClientNameColumnVisibility = InArgs._ClientNameColumnVisibility;
-	OperationColumnVisibility = InArgs._OperationColumnVisibility;
-	PackageColumnVisibility = InArgs._PackageColumnVisibility;
 	ConnectionActivitiesVisibility = InArgs._ConnectionActivitiesVisibility;
 	LockActivitiesVisibility = InArgs._LockActivitiesVisibility;
 	PackageActivitiesVisibility = InArgs._PackageActivitiesVisibility;
@@ -378,46 +230,7 @@ void SConcertSessionActivities::Construct(const FArguments& InArgs)
 	SearchTextFilter = MakeShared<TTextFilter<const FConcertSessionActivity&>>(TTextFilter<const FConcertSessionActivity&>::FItemToStringArray::CreateSP(this, &SConcertSessionActivities::PopulateSearchStrings));
 	SearchTextFilter->OnChanged().AddSP(this, &SConcertSessionActivities::OnActivityFilterUpdated);
 
-	// Set the initial filter state.
 	ActiveFilterFlags = QueryActiveActivityFilters();
-
-	// Create the table header. (Setting visibility on the column itself doesn't show/hide the column as one would expect, unfortunately)
-	TSharedRef<SHeaderRow> HeaderRow = SNew(SHeaderRow);
-	if (InArgs._ClientAvatarColorColumnVisibility.Get() == EVisibility::Visible)
-	{
-		HeaderRow->AddColumn(SHeaderRow::Column(ConcertSessionActivityUtils::AvatarColorColumnId)
-			.DefaultLabel(INVTEXT(""))
-			.ManualWidth(8));
-	}
-
-	HeaderRow->AddColumn(SHeaderRow::Column(ConcertSessionActivityUtils::DateTimeColumnId)
-		.DefaultLabel(LOCTEXT("DateTime", "Date/Time"))
-		.ManualWidth(160));
-
-	if (InArgs._ClientNameColumnVisibility.Get() == EVisibility::Visible)
-	{
-		HeaderRow->AddColumn(SHeaderRow::Column(ConcertSessionActivityUtils::ClientNameColumnId)
-			.DefaultLabel(LOCTEXT("Client", "Client"))
-			.ManualWidth(80));
-	}
-
-	if (InArgs._OperationColumnVisibility.Get() == EVisibility::Visible)
-	{
-		HeaderRow->AddColumn(SHeaderRow::Column(ConcertSessionActivityUtils::OperationColumnId)
-			.DefaultLabel(LOCTEXT("Operation", "Operation"))
-			.ManualWidth(160));
-	}
-
-	if (InArgs._PackageColumnVisibility.Get() == EVisibility::Visible)
-	{
-		HeaderRow->AddColumn(SHeaderRow::Column(ConcertSessionActivityUtils::PackageColumnId)
-			.DefaultLabel(LOCTEXT("Package", "Package"))
-			.ManualWidth(200));
-	}
-
-	HeaderRow->AddColumn(SHeaderRow::Column(ConcertSessionActivityUtils::SummaryColumnId)
-		.DefaultLabel(LOCTEXT("Summary", "Summary")));
-
 	ChildSlot
 	[
 		SNew(SSplitter)
@@ -443,7 +256,7 @@ void SConcertSessionActivities::Construct(const FArguments& InArgs)
 					.AllowOverscroll(EAllowOverscroll::No)
 					.OnListViewScrolled(this, &SConcertSessionActivities::OnListViewScrolled)
 					.OnSelectionChanged(this, &SConcertSessionActivities::OnListViewSelectionChanged)
-					.HeaderRow(HeaderRow)
+					.HeaderRow(CreateHeaderRow())
 				]
 			]
 
@@ -536,6 +349,52 @@ void SConcertSessionActivities::Construct(const FArguments& InArgs)
 	}
 }
 
+TSharedRef<SHeaderRow> SConcertSessionActivities::CreateHeaderRow()
+{
+	using namespace UE::ConcertSharedSlate;
+
+	// Before sort because columns can be placed before
+	const bool bContainsDateTime = Columns.FindByPredicate([](const FActivityColumn& Column) { return Column.ColumnId ==  ActivityColumn::DateTimeColumnId; }) != nullptr; 
+	checkf(!bContainsDateTime, TEXT("DateTime is already created by SConcertSessionActivities!"))
+	Columns.Add(ActivityColumn::DateTime());
+	
+	Columns.Sort([](const FActivityColumn& Left, const FActivityColumn& Right) { return Left.GetColumnSortOrderValue() < Right.GetColumnSortOrderValue(); });
+
+	// Summary is always last
+	const bool bContainsSummary = Columns.FindByPredicate([](const FActivityColumn& Column) { return Column.ColumnId ==  ActivityColumn::SummaryColumnId; }) != nullptr; 
+	checkf(!bContainsSummary, TEXT("Summary is already created by SConcertSessionActivities!"))
+	Columns.Add(ActivityColumn::Summary());
+	
+	const TSharedRef<SHeaderRow> HeaderRow = SNew(SHeaderRow);
+	for (FActivityColumn& Column : Columns)
+	{
+		HeaderRow->AddColumn(Column);
+	}
+	
+	return HeaderRow;
+}
+
+TSharedRef<ITableRow> SConcertSessionActivities::OnGenerateActivityRowWidget(TSharedPtr<FConcertSessionActivity> Activity, const TSharedRef<STableViewBase>& OwnerTable)
+{
+	const TOptional<FConcertClientInfo> ActivityClient = GetActivityUserFn.IsBound() ? GetActivityUserFn.Execute(Activity->Activity.EndpointId) : TOptional<FConcertClientInfo>{};
+	const SConcertSessionActivityRow::FGetColumn ColumnGetter = SConcertSessionActivityRow::FGetColumn::CreateLambda([this](const FName& ColumnId) -> const FActivityColumn*
+	{
+		for (const FActivityColumn& Column : Columns)
+		{
+			if (Column.ColumnId == ColumnId)
+			{
+				return &Column;
+			}
+		}
+
+		checkNoEntry();
+		return nullptr;
+	});
+	
+	return SNew(SConcertSessionActivityRow, SharedThis(this), Activity.ToSharedRef(), ColumnGetter, ActivityClient, OwnerTable)
+		.OnMakeColumnOverlayWidget(MakeColumnOverlayWidgetFn);
+}
+
 void SConcertSessionActivities::Tick(const FGeometry& AllottedGeometry, const double InCurrentTime, const float InDeltaTime)
 {
 	EConcertActivityFilterFlags LatestFilterFlags = QueryActiveActivityFilters();
@@ -870,37 +729,12 @@ FText SConcertSessionActivities::UpdateTextFilter(const FText& InFilterText)
 	return SearchTextFilter->GetFilterErrorText();
 }
 
-void SConcertSessionActivities::PopulateSearchStrings(const FConcertSessionActivity& Activity, TArray<FString>& OutSearchStrings) const
+void SConcertSessionActivities::PopulateSearchStrings(const FConcertSessionActivity& Activity, TArray<FString>& OutSearchStrings)
 {
-	FText ClientName = GetActivityUserFn.IsBound()
-		? ConcertSessionActivityUtils::GetClientName(GetActivityUserFn.Execute(Activity.Activity.EndpointId)) : FText::GetEmpty();
-
-	OutSearchStrings.Add(ConcertSessionActivityUtils::GetActivityDateTime(Activity, TimeFormat.Get()).ToString());
-	OutSearchStrings.Add(ConcertSessionActivityUtils::GetSummary(Activity, ClientName, /*bAsRichText*/false).ToString());
-
-	if (ClientNameColumnVisibility.Get() == EVisibility::Visible)
+	for (const FActivityColumn& Column : Columns)
 	{
-		OutSearchStrings.Add(ClientName.ToString());
+		Column.ExecutePopulateSearchString(SharedThis(this), Activity, OutSearchStrings);
 	}
-
-	if (OperationColumnVisibility.Get() == EVisibility::Visible)
-	{
-		OutSearchStrings.Add(ConcertSessionActivityUtils::GetOperationName(Activity).ToString());
-	}
-
-	if (PackageColumnVisibility.Get() == EVisibility::Visible)
-	{
-		OutSearchStrings.Add(ConcertSessionActivityUtils::GetPackageName(Activity).ToString());
-	}
-}
-
-TSharedRef<ITableRow> SConcertSessionActivities::OnGenerateActivityRowWidget(TSharedPtr<FConcertSessionActivity> Activity, const TSharedRef<STableViewBase>& OwnerTable)
-{
-	const TOptional<FConcertClientInfo> ActivityClient = GetActivityUserFn.IsBound() ? GetActivityUserFn.Execute(Activity->Activity.EndpointId) : TOptional<FConcertClientInfo>{};
-	return SNew(SConcertSessionActivityRow, Activity, ActivityClient, OwnerTable)
-		.TimeFormat(TimeFormat)
-		.HighlightText(HighlightText)
-		.OnMakeColumnOverlayWidget(MakeColumnOverlayWidgetFn);
 }
 
 TSharedPtr<FConcertSessionActivity> SConcertSessionActivities::GetSelectedActivity() const
@@ -928,7 +762,8 @@ TSharedPtr<FConcertSessionActivity> SConcertSessionActivities::GetMostRecentActi
 
 bool SConcertSessionActivities::IsLastColumn(const FName& ColumnId) const
 {
-	return ColumnId == ConcertSessionActivityUtils::SummaryColumnId; // Summary column is always visible and always the last.
+	// Summary column is always visible and always the last.
+	return ColumnId == UE::ConcertSharedSlate::ActivityColumn::SummaryColumnId; 
 }
 
 void SConcertSessionActivities::DisplayTransactionDetails(const FConcertSessionActivity& Activity, const FConcertTransactionEventBase& InTransaction)
diff --git a/Engine/Plugins/Developer/Concert/ConcertUI/ConcertSharedSlate/Source/ConcertSharedSlate/Private/Session/Activity/SessionActivityUtils.cpp b/Engine/Plugins/Developer/Concert/ConcertUI/ConcertSharedSlate/Source/ConcertSharedSlate/Private/Session/Activity/SessionActivityUtils.cpp
new file mode 100644
--- /dev/null
+++ b/Engine/Plugins/Developer/Concert/ConcertUI/ConcertSharedSlate/Source/ConcertSharedSlate/Private/Session/Activity/SessionActivityUtils.cpp
@@ -0,0 +1,77 @@
+﻿// Copyright Epic Games, Inc. All Rights Reserved.
+
+#include "SessionActivityUtils.h"
+
+#include "ConcertSyncSessionTypes.h"
+#include "ConcertWorkspaceData.h"
+
+#define LOCTEXT_NAMESPACE "SConcertSessionActivities"
+
+FText UE::ConcertSharedSlate::Private::GetOperationName(const FConcertSessionActivity& Activity)
+{
+	if (const FConcertSyncTransactionActivitySummary* TransactionSummary = Activity.ActivitySummary.Cast<FConcertSyncTransactionActivitySummary>())
+	{
+		return TransactionSummary->TransactionTitle;
+	}
+
+	if (const FConcertSyncPackageActivitySummary* PackageSummary = Activity.ActivitySummary.Cast<FConcertSyncPackageActivitySummary>())
+	{
+		auto GetSavedOperationNameFn = [](const FConcertSyncPackageActivitySummary* PackageSummary)
+		{
+			if (PackageSummary->bAutoSave)
+			{
+				return LOCTEXT("AutoSavePackageOperation", "Auto-Save Package");
+			}
+			return PackageSummary->bPreSave ? LOCTEXT("PreSavePackageOperation", "Pre-Save Package") : LOCTEXT("SavePackageOperation", "Save Package");
+		};
+
+		switch (PackageSummary->PackageUpdateType)
+		{
+		case EConcertPackageUpdateType::Added   : return LOCTEXT("NewPackageOperation",    "New Package");
+		case EConcertPackageUpdateType::Deleted : return LOCTEXT("DeletePackageOperation", "Delete Package");
+		case EConcertPackageUpdateType::Renamed : return LOCTEXT("RenamePackageOperation", "Rename Package");
+		case EConcertPackageUpdateType::Saved   : return GetSavedOperationNameFn(PackageSummary);
+		case EConcertPackageUpdateType::Dummy   : return LOCTEXT("DiscardPackageOperation", "Discard Changes");
+		default: break;
+		}
+	}
+
+	if (const FConcertSyncConnectionActivitySummary* ConnectionSummary = Activity.ActivitySummary.Cast<FConcertSyncConnectionActivitySummary>())
+	{
+		switch (ConnectionSummary->ConnectionEventType)
+		{
+		case EConcertSyncConnectionEventType::Connected:    return LOCTEXT("JoinOperation", "Join Session");
+		case EConcertSyncConnectionEventType::Disconnected: return LOCTEXT("LeaveOperation", "Leave Session");
+		default: break;
+		}
+	}
+
+	if (const FConcertSyncLockActivitySummary* LockSummary = Activity.ActivitySummary.Cast<FConcertSyncLockActivitySummary>())
+	{
+		switch (LockSummary->LockEventType)
+		{
+		case EConcertSyncLockEventType::Locked:   return LOCTEXT("LockOperation", "Lock");
+		case EConcertSyncLockEventType::Unlocked: return LOCTEXT("UnlockOperation", "Unlock");
+		default: break;
+		}
+	}
+
+	return FText::GetEmpty();
+}
+
+FText UE::ConcertSharedSlate::Private::GetPackageName(const FConcertSessionActivity& Activity)
+{
+	if (const FConcertSyncPackageActivitySummary* Summary = Activity.ActivitySummary.Cast<FConcertSyncPackageActivitySummary>())
+	{
+		return FText::FromName(Summary->PackageName);
+	}
+
+	if (const FConcertSyncTransactionActivitySummary* Summary = Activity.ActivitySummary.Cast<FConcertSyncTransactionActivitySummary>())
+	{
+		return FText::FromName(Summary->PrimaryPackageName);
+	}
+
+	return FText::GetEmpty();
+}
+
+#undef LOCTEXT_NAMESPACE
\ No newline at end of file
diff --git a/Engine/Plugins/Developer/Concert/ConcertUI/ConcertSharedSlate/Source/ConcertSharedSlate/Private/Session/Activity/SessionActivityUtils.h b/Engine/Plugins/Developer/Concert/ConcertUI/ConcertSharedSlate/Source/ConcertSharedSlate/Private/Session/Activity/SessionActivityUtils.h
new file mode 100644
--- /dev/null
+++ b/Engine/Plugins/Developer/Concert/ConcertUI/ConcertSharedSlate/Source/ConcertSharedSlate/Private/Session/Activity/SessionActivityUtils.h
@@ -0,0 +1,13 @@
+﻿// Copyright Epic Games, Inc. All Rights Reserved.
+
+#pragma once
+
+#include "CoreMinimal.h"
+
+struct FConcertSessionActivity;
+
+namespace UE::ConcertSharedSlate::Private
+{
+	FText GetOperationName(const FConcertSessionActivity& Activity);
+	FText GetPackageName(const FConcertSessionActivity& Activity);
+}
diff --git a/Engine/Plugins/Developer/Concert/ConcertUI/ConcertSharedSlate/Source/ConcertSharedSlate/Private/Session/History/SSessionHistory.cpp b/Engine/Plugins/Developer/Concert/ConcertUI/ConcertSharedSlate/Source/ConcertSharedSlate/Private/Session/History/SSessionHistory.cpp
--- a/Engine/Plugins/Developer/Concert/ConcertUI/ConcertSharedSlate/Source/ConcertSharedSlate/Private/Session/History/SSessionHistory.cpp
+++ b/Engine/Plugins/Developer/Concert/ConcertUI/ConcertSharedSlate/Source/ConcertSharedSlate/Private/Session/History/SSessionHistory.cpp
@@ -4,6 +4,8 @@
 
 #include "ConcertMessageData.h"
 #include "IConcertSession.h"
+
+#include "Session/Activity/PredefinedActivityColumns.h"
 #include "Session/Activity/SConcertSessionActivities.h"
 
 #include "Widgets/Input/SSearchBox.h"
@@ -41,6 +43,8 @@ namespace ConcertSessionHistoryUI
 
 void SSessionHistory::Construct(const FArguments& InArgs)
 {
+	using namespace UE::ConcertSharedSlate;
+	
 	PackageNameFilter = InArgs._PackageFilter;
 
 	ActivityMap.Reserve(MaximumNumberOfActivities);
@@ -52,10 +56,11 @@ void SSessionHistory::Construct(const FArguments& InArgs)
 		.OnMapActivityToClient(this, &SSessionHistory::GetClientInfo)
 		.HighlightText(this, &SSessionHistory::HighlightSearchedText)
 		.TimeFormat(ActivityListViewOptions.Get(), &FConcertSessionActivitiesOptions::GetTimeFormat)
-		.ClientNameColumnVisibility(EVisibility::Visible)
-		.ClientAvatarColorColumnVisibility(EVisibility::Visible)
-		.OperationColumnVisibility(EVisibility::Visible)
-		.PackageColumnVisibility(EVisibility::Collapsed)
+		.Columns({
+			ActivityColumn::AvatarColor(),
+			ActivityColumn::ClientName(),
+			ActivityColumn::Operation()
+		})
 		.ConnectionActivitiesVisibility(ActivityListViewOptions.Get(), &FConcertSessionActivitiesOptions::GetConnectionActivitiesVisibility)
 		.LockActivitiesVisibility(ActivityListViewOptions.Get(), &FConcertSessionActivitiesOptions::GetLockActivitiesVisibility)
 		.PackageActivitiesVisibility(ActivityListViewOptions.Get(), &FConcertSessionActivitiesOptions::GetPackageActivitiesVisibility)
diff --git a/Engine/Plugins/Developer/Concert/ConcertUI/ConcertSharedSlate/Source/ConcertSharedSlate/Public/SConcertSessionRecovery.h b/Engine/Plugins/Developer/Concert/ConcertUI/ConcertSharedSlate/Source/ConcertSharedSlate/Public/SConcertSessionRecovery.h
--- a/Engine/Plugins/Developer/Concert/ConcertUI/ConcertSharedSlate/Source/ConcertSharedSlate/Public/SConcertSessionRecovery.h
+++ b/Engine/Plugins/Developer/Concert/ConcertUI/ConcertSharedSlate/Source/ConcertSharedSlate/Public/SConcertSessionRecovery.h
@@ -19,10 +19,10 @@ class CONCERTSHAREDSLATE_API SConcertSessionRecovery : public SCompoundWidget
 {
 public:
 	SLATE_BEGIN_ARGS(SConcertSessionRecovery)
-		: _ClientAvatarColorColumnVisibility(EVisibility::Collapsed)
-		, _ClientNameColumnVisibility(EVisibility::Collapsed)
-		, _OperationColumnVisibility(EVisibility::Visible)
-		, _PackageColumnVisibility(EVisibility::Collapsed)
+		: _WithClientAvatarColorColumn(false)
+		, _WithClientNameColumn(false)
+		, _WithOperationColumn(true)
+		, _WithPackageColumn(false)
 		, _DetailsAreaVisibility(EVisibility::Collapsed)
 		, _IsConnectionActivityFilteringEnabled(false)
 		, _IsLockActivityFilteringEnabled(false)
@@ -51,16 +51,16 @@ public:
 		SLATE_ARGUMENT(TFunction<void()>, OnCancel)
 
 		/** Show/hide the column displaying the avatar color of the client who performed the activity. */
-		SLATE_ARGUMENT(EVisibility, ClientAvatarColorColumnVisibility)
+		SLATE_ARGUMENT(bool, WithClientAvatarColorColumn)
 
 		/** Show/hide the column showing the display name of the client who performed the activity. */
-		SLATE_ARGUMENT(EVisibility, ClientNameColumnVisibility)
+		SLATE_ARGUMENT(bool, WithClientNameColumn)
 
 		/** Show/hide the column displaying the operation represented by the activity. */
-		SLATE_ARGUMENT(EVisibility, OperationColumnVisibility)
+		SLATE_ARGUMENT(bool, WithOperationColumn)
 
 		/** Show/hide the column displaying affected package. */
-		SLATE_ARGUMENT(EVisibility, PackageColumnVisibility)
+		SLATE_ARGUMENT(bool, WithPackageColumn)
 
 		/** Show/hide the details area widget. */
 		SLATE_ARGUMENT(EVisibility, DetailsAreaVisibility)
@@ -112,6 +112,9 @@ public:
 	TSharedPtr<FConcertSessionActivity> GetMostRecentActivity() const;
 
 private:
+
+	TArray<FActivityColumn> MakeColumns(const FArguments& InArgs) const;
+	
 	void OnSearchTextChanged(const FText& InFilterText);
 	void OnSearchTextCommitted(const FText& InFilterText, ETextCommit::Type CommitType);
 	FText HighlightSearchText() const;
diff --git a/Engine/Plugins/Developer/Concert/ConcertUI/ConcertSharedSlate/Source/ConcertSharedSlate/Public/Session/Activity/ActivityColumn.h b/Engine/Plugins/Developer/Concert/ConcertUI/ConcertSharedSlate/Source/ConcertSharedSlate/Public/Session/Activity/ActivityColumn.h
new file mode 100644
--- /dev/null
+++ b/Engine/Plugins/Developer/Concert/ConcertUI/ConcertSharedSlate/Source/ConcertSharedSlate/Public/Session/Activity/ActivityColumn.h
@@ -0,0 +1,53 @@
+﻿// Copyright Epic Games, Inc. All Rights Reserved.
+
+#pragma once
+
+#include "CoreMinimal.h"
+#include "Widgets/Views/SHeaderRow.h"
+
+class SConcertSessionActivities;
+struct FConcertSessionActivity;
+
+/** A column in SConcertSessionActivities */
+class CONCERTSHAREDSLATE_API FActivityColumn final : public SHeaderRow::FColumn
+{
+public:
+
+	/** Generates the widget content for a column in a row */
+	DECLARE_DELEGATE_ThreeParams(FGenerateColumnWidget,
+		const TSharedRef<SConcertSessionActivities>& Owner,
+		const TSharedRef<FConcertSessionActivity>&,
+		SOverlay::FScopedWidgetSlotArguments& Slot
+		)
+
+	/** Appends search terms to ExistingSearchStrings. */
+	DECLARE_DELEGATE_ThreeParams(FPopulateSearchString,
+		const TSharedRef<SConcertSessionActivities>& Owner,
+		const FConcertSessionActivity& Activity,
+		TArray<FString>& ExistingSearchStrings
+		)
+	
+	FActivityColumn(const FArguments& InArgs)
+		: FColumn(InArgs)
+	{}
+
+	/** Called to generate the column's widget for a row */
+	FActivityColumn& GenerateColumnWidget(FGenerateColumnWidget Callback) { GenerateColumnWidgetCallback = MoveTemp(Callback); return *this; }
+	/** Called to populate the search terms when rows are search */
+	FActivityColumn& PopulateSearchString(FPopulateSearchString Callback) { PopulateSearchStringCallback = MoveTemp(Callback); return *this; }
+
+	/** Determines whether this column is the first, etc. */
+	FActivityColumn& ColumnSortOrder(int32 NewValue) { ColumnSortOrderValue = NewValue; return *this; }
+	
+	void BuildColumnWidget(const TSharedRef<SConcertSessionActivities>& Owner, const TSharedRef<FConcertSessionActivity>& Activity, SOverlay::FScopedWidgetSlotArguments& Slot) const;
+	void ExecutePopulateSearchString(const TSharedRef<SConcertSessionActivities>& Owner, const FConcertSessionActivity& Activity, TArray<FString>& ExistingSearchStrings) const;
+	int32 GetColumnSortOrderValue() const { return ColumnSortOrderValue; }
+	
+private:
+
+	FGenerateColumnWidget GenerateColumnWidgetCallback;
+	FPopulateSearchString PopulateSearchStringCallback;
+
+	/** Determines whether this column is the first, etc. */
+	int32 ColumnSortOrderValue = 0;
+};
diff --git a/Engine/Plugins/Developer/Concert/ConcertUI/ConcertSharedSlate/Source/ConcertSharedSlate/Public/Session/Activity/PredefinedActivityColumns.h b/Engine/Plugins/Developer/Concert/ConcertUI/ConcertSharedSlate/Source/ConcertSharedSlate/Public/Session/Activity/PredefinedActivityColumns.h
new file mode 100644
--- /dev/null
+++ b/Engine/Plugins/Developer/Concert/ConcertUI/ConcertSharedSlate/Source/ConcertSharedSlate/Public/Session/Activity/PredefinedActivityColumns.h
@@ -0,0 +1,39 @@
+﻿// Copyright Epic Games, Inc. All Rights Reserved.
+
+#pragma once
+
+#include "CoreMinimal.h"
+#include "SConcertSessionActivities.h"
+
+class FActivityColumn;
+
+/** Contains factories for commonly used columns. By default, they maintain the relative order specified by EPredefinedColumnOrder. */
+namespace UE::ConcertSharedSlate::ActivityColumn
+{
+	enum class EPredefinedColumnOrder : int32
+	{
+		AvatarColor = -20,
+		DateTime = -10,
+		ClientName = 10,
+		Operation = 20,
+		Package = 30
+		// No Summary entry because it is forced to be last
+	};
+	
+	extern const FName DateTimeColumnId;
+	extern const FName OperationColumnId;
+	extern const FName AvatarColorColumnId;
+	extern const FName ClientNameColumnId;
+	extern const FName PackageColumnId;
+	extern const FName SummaryColumnId;
+	
+	// Required
+	FActivityColumn DateTime();
+	FActivityColumn Summary();
+
+	// Optional
+	CONCERTSHAREDSLATE_API FActivityColumn AvatarColor();
+	CONCERTSHAREDSLATE_API FActivityColumn ClientName();
+	CONCERTSHAREDSLATE_API FActivityColumn Operation();
+	CONCERTSHAREDSLATE_API FActivityColumn Package();
+}
diff --git a/Engine/Plugins/Developer/Concert/ConcertUI/ConcertSharedSlate/Source/ConcertSharedSlate/Public/Session/Activity/SConcertSessionActivities.h b/Engine/Plugins/Developer/Concert/ConcertUI/ConcertSharedSlate/Source/ConcertSharedSlate/Public/Session/Activity/SConcertSessionActivities.h
--- a/Engine/Plugins/Developer/Concert/ConcertUI/ConcertSharedSlate/Source/ConcertSharedSlate/Public/Session/Activity/SConcertSessionActivities.h
+++ b/Engine/Plugins/Developer/Concert/ConcertUI/ConcertSharedSlate/Source/ConcertSharedSlate/Public/Session/Activity/SConcertSessionActivities.h
@@ -3,7 +3,7 @@
 #pragma once
 
 #include "CoreMinimal.h"
-#include "SlateFwd.h"
+#include "ActivityColumn.h"
 #include "Async/Future.h"
 #include "Misc/TextFilter.h"
 #include "Widgets/DeclarativeSyntaxSupport.h"
@@ -20,7 +20,6 @@ class SPackageDetails;
 class STextBlock;
 class SExpendableArea;
 
-
 /** Filters for the concert session activity view. */
 enum class EConcertActivityFilterFlags
 {
@@ -37,7 +36,7 @@ ENUM_CLASS_FLAGS(EConcertActivityFilterFlags);
 class CONCERTSHAREDSLATE_API SConcertSessionActivities : public SCompoundWidget
 {
 public:
-	
+
 	/** Defines how the time should be displayed in the date/time column. */
 	enum class ETimeFormat
 	{
@@ -54,7 +53,6 @@ public:
 	/** Returns the transaction event corresponding the specified activity.*/
 	DECLARE_DELEGATE_RetVal_OneParam(TFuture<TOptional<FConcertSyncTransactionEvent>>, FGetTransactionEvent, const FConcertSessionActivity& /*Activity*/)
 	
-
 	/** Returns the package event corresponding to the package activity. */
 	DECLARE_DELEGATE_RetVal_TwoParams(bool, FGetPackageEvent, const FConcertSessionActivity& /*Activity*/, FConcertSyncPackageEventMetaData& /*OutEvent*/);
 
@@ -67,10 +65,6 @@ public:
 		, _OnMapActivityToClient()
 		, _OnMakeColumnOverlayWidget()
 		, _TimeFormat(ETimeFormat::Relative)
-		, _ClientAvatarColorColumnVisibility(EVisibility::Hidden)
-		, _ClientNameColumnVisibility(EVisibility::Hidden)
-		, _OperationColumnVisibility(EVisibility::Hidden)
-		, _PackageColumnVisibility(EVisibility::Hidden)
 		, _ConnectionActivitiesVisibility(EVisibility::Hidden)
 		, _LockActivitiesVisibility(EVisibility::Hidden)
 		, _PackageActivitiesVisibility(EVisibility::Visible)
@@ -93,6 +87,9 @@ public:
 
 		/** If bound, invoked when generating a row to add an overlay to a column. */
 		SLATE_EVENT(FMakeColumnOverlayWidgetFunc, OnMakeColumnOverlayWidget)
+	
+		/** The columns to show for each activity; they are shown after date time */
+		SLATE_ARGUMENT(TArray<FActivityColumn>, Columns)
 
 		/** Highlight the returned text in the view. */
 		SLATE_ATTRIBUTE(FText, HighlightText)
@@ -100,18 +97,6 @@ public:
 		/** Defines how time should be displayed (relative vs absolute). */
 		SLATE_ATTRIBUTE(ETimeFormat, TimeFormat)
 
-		/** Show/hide the column displaying the avatar color of the client who performed the activity. */
-		SLATE_ATTRIBUTE(EVisibility, ClientAvatarColorColumnVisibility)
-
-		/** Show/hide the column showing the display name of the client who performed the activity. */
-		SLATE_ATTRIBUTE(EVisibility, ClientNameColumnVisibility)
-
-		/** Show/hide the column showing the operation name. */
-		SLATE_ATTRIBUTE(EVisibility, OperationColumnVisibility)
-
-		/** Show/hide the column showing the affected package. */
-		SLATE_ATTRIBUTE(EVisibility, PackageColumnVisibility)
-
 		/** Show/hide connection activities. */
 		SLATE_ATTRIBUTE(EVisibility, ConnectionActivitiesVisibility)
 
@@ -176,59 +161,47 @@ public:
 	/** Asks the view to update the text search filter. */
 	FText UpdateTextFilter(const FText& InFilterText);
 
+	const TAttribute<FText>& GetHighlightText() const { return HighlightText; }
+	const TAttribute<ETimeFormat>& GetTimeFormat() const { return TimeFormat; }
+	TOptional<FConcertClientInfo> GetClientInfo(const FGuid& Guid) const { return GetActivityUserFn.IsBound() ? GetActivityUserFn.Execute(Guid) : TOptional<FConcertClientInfo>{}; }
+
 private:
-	/** Generates a row widget in the table view. */
-	TSharedRef<ITableRow> OnGenerateActivityRowWidget(TSharedPtr<FConcertSessionActivity> Item, const TSharedRef<STableViewBase>& OwnerTable);
 
-	/** Fetches more activities. */
+	// Widget creation
+	TSharedRef<SHeaderRow> CreateHeaderRow();
+	TSharedRef<ITableRow> OnGenerateActivityRowWidget(TSharedPtr<FConcertSessionActivity> Item, const TSharedRef<STableViewBase>& OwnerTable);
 	void FetchActivities();
 
-	/** Queries the active set of activity filter flags. */
+	// Filtering
 	EConcertActivityFilterFlags QueryActiveActivityFilters() const;
-
-	/** Invoked when the filter changes and the displayed activity list must be reevaluated. */
 	void OnActivityFilterUpdated();
-
-	/** Returns true if the specified activity passes the view filters. */
 	bool PassesFilters(const FConcertSessionActivity& Activity);
 
-	/** Invoked when the list view scrolled to fetch more data from the activity provider. */
+	// List view
 	void OnListViewScrolled(double InScrollOffset);
-
-	/** Invoked when the list view selection changes. */
 	void OnListViewSelectionChanged(TSharedPtr<FConcertSessionActivity> InActivity, ESelectInfo::Type SelectInfo);
 
-	/** Convert the item into a list of text element used search bars. */
-	void PopulateSearchStrings(const FConcertSessionActivity& Item, TArray<FString>& OutSearchStrings) const;
-
-	/** Returns whether the details expandable area should be visible or not. */
-	EVisibility GetDetailAreaVisibility() const { return DetailsAreaVisibility; }
-
-	/** Manages how much space is used by the 'Details' area with respect to its expansion state. */
-	SSplitter::ESizeRule GetDetailsAreaSizeRule() const { return bDetailsAreaExpanded ? SSplitter::ESizeRule::FractionOfParent : SSplitter::ESizeRule::SizeToContent; }
+	// Search
+	void PopulateSearchStrings(const FConcertSessionActivity& Item, TArray<FString>& OutSearchStrings);
 
-	/** Invoked when the 'Details' expandable area is expanded/collapsed. Affects how the 'Details' area size is accounted in the splitter. */
+	// Details
 	void OnDetailsAreaExpansionChanged(bool bExpanded);
-
-	/** Displays what properties/objects changed during a transaction in the details area when a transaction activity is selected in the list. */
-	void DisplayTransactionDetails(const FConcertSessionActivity& Activity, const FConcertTransactionEventBase& InTransaction);
-
-	/** Displays what changed in a package in the details area when a package activity is selected in the list. */
-	void DisplayPackageDetails(const FConcertSessionActivity& Activity, int64 PackageRevision, const FConcertPackageInfo& PackageInfo);
-
-	/** Update the detail area to display the proper panel and details. */
 	void UpdateDetailArea(TSharedPtr<FConcertSessionActivity> InSelectedActivity);
-
-	/** Change the visibility of the details panels, keeping only one visible. */
 	void SetDetailsPanelVisibility(const SWidget* VisiblePanel);
-
-	/** Returns the text displayed to the user when there are no details to display. */
+	void DisplayTransactionDetails(const FConcertSessionActivity& Activity, const FConcertTransactionEventBase& InTransaction);
+	void DisplayPackageDetails(const FConcertSessionActivity& Activity, int64 PackageRevision, const FConcertPackageInfo& PackageInfo);
 	FText GetNoDetailsText() const;
-
+	EVisibility GetDetailAreaVisibility() const { return DetailsAreaVisibility; }
+	SSplitter::ESizeRule GetDetailsAreaSizeRule() const { return bDetailsAreaExpanded ? SSplitter::ESizeRule::FractionOfParent : SSplitter::ESizeRule::SizeToContent; }
+	
 	/** Invoked after all widget ticks get processed. Used to check if the user scrolled or not. */
 	void OnPostTick(float);
 
 private:
+
+	/** The columns that are being displayed */
+	TArray<FActivityColumn> Columns;
+	
 	/** List of all activities (including the filtered out ones). */
 	TArray<TSharedPtr<FConcertSessionActivity>> AllActivities;
 
@@ -247,15 +220,6 @@ private:
 	/** Whether the time should be displayed as relative (9 seconds ago) or as absolute (July 10, 2019 - 10:20:10) */
 	TAttribute<ETimeFormat> TimeFormat;
 
-	/** Whether the client name column is displayed. */
-	TAttribute<EVisibility> ClientNameColumnVisibility;
-
-	/** Whether the operation column is displayed. */
-	TAttribute<EVisibility> OperationColumnVisibility;
-
-	/** Whether the package column is displayed. */
-	TAttribute<EVisibility> PackageColumnVisibility;
-
 	/** Whether the join/leave session activities are displayed. */
 	TAttribute<EVisibility> ConnectionActivitiesVisibility;
 
@@ -281,7 +245,7 @@ private:
 	bool bUserScrolling = false;
 
 	/** Defines which activity types are currently filtered out from the view. */
-	EConcertActivityFilterFlags ActiveFilterFlags;
+	EConcertActivityFilterFlags ActiveFilterFlags = EConcertActivityFilterFlags::ShowAll;
 
 	/** Used to fetch more activities from an abstract source. Usually mutually exclusive with Append() function. May not be bound. */
 	FFetchActivitiesFunc FetchActivitiesFn;
diff --git a/Engine/Plugins/Developer/Concert/DisasterRecoveryClient/Source/DisasterRecoveryClient/Private/SDisasterRecoveryHub.cpp b/Engine/Plugins/Developer/Concert/DisasterRecoveryClient/Source/DisasterRecoveryClient/Private/SDisasterRecoveryHub.cpp
--- a/Engine/Plugins/Developer/Concert/DisasterRecoveryClient/Source/DisasterRecoveryClient/Private/SDisasterRecoveryHub.cpp
+++ b/Engine/Plugins/Developer/Concert/DisasterRecoveryClient/Source/DisasterRecoveryClient/Private/SDisasterRecoveryHub.cpp
@@ -597,10 +597,10 @@ TSharedRef<SWidget> SDisasterRecoveryHub::MakeSessionActivityView()
 {
 	return SAssignNew(ActivityView, SConcertSessionRecovery)
 		.OnFetchActivities(this, &SDisasterRecoveryHub::FetchActivities)
-		.ClientAvatarColorColumnVisibility(EVisibility::Hidden) // Disaster recovery has only one user, the local one.
-		.ClientNameColumnVisibility(EVisibility::Hidden)
-		.OperationColumnVisibility(EVisibility::Visible)
-		.PackageColumnVisibility(EVisibility::Visible)
+		.WithClientAvatarColorColumn(false) // Disaster recovery has only one user, the local one.
+		.WithClientNameColumn(false)
+		.WithOperationColumn(true)
+		.WithPackageColumn(true)
 		.DetailsAreaVisibility(ShouldDisplayActivityDetails() ? EVisibility::Visible : EVisibility::Collapsed)
 		.IsConnectionActivityFilteringEnabled(false) // Not valuable and ignored by Disaster Recovery (DR). DR sessions are for the local user only and join/leave are not recoverable but this will also ignore when the local user joins/leaves a Multi-User session.
 		.IsLockActivityFilteringEnabled(false)       // Not valuable and ignored by Disaster Recovery (DR). DR sessions are for the local user only but lock/unlock are not recoverable but this will also ignore when the local user locks/unlocks assets in a Multi-User session.
diff --git a/Engine/Plugins/Developer/Concert/MultiUserClient/Source/MultiUserClient/Private/Widgets/SConcertBrowser.cpp b/Engine/Plugins/Developer/Concert/MultiUserClient/Source/MultiUserClient/Private/Widgets/SConcertBrowser.cpp
--- a/Engine/Plugins/Developer/Concert/MultiUserClient/Source/MultiUserClient/Private/Widgets/SConcertBrowser.cpp
+++ b/Engine/Plugins/Developer/Concert/MultiUserClient/Source/MultiUserClient/Private/Widgets/SConcertBrowser.cpp
@@ -421,10 +421,10 @@ void FConcertClientSessionBrowserController::RestoreSession(const FGuid& ServerA
 		.OnFetchActivities_Lambda(ReadActivitiesFn)
 		.OnMapActivityToClient_Lambda(GetActivityClientInfoFn)
 		.OnRestore(OnAcceptRestoreFn)
-		.ClientNameColumnVisibility(EVisibility::Visible)
-		.ClientAvatarColorColumnVisibility(EVisibility::Visible)
-		.OperationColumnVisibility(EVisibility::Visible)
-		.PackageColumnVisibility(EVisibility::Hidden) // Even tough the column is not present, the tooltips and summary contains the affected package.
+		.WithClientAvatarColorColumn(true)
+		.WithClientNameColumn(true)
+		.WithOperationColumn(true)
+		.WithPackageColumn(false) // Even tough the column is not present, the tooltips and summary contains the affected package.
 		.DetailsAreaVisibility(bRequestActivityDetails ? EVisibility::Visible : EVisibility::Collapsed) // The activity stream was configured to pull the activity details.
 		.IsConnectionActivityFilteringEnabled(true)
 		.IsLockActivityFilteringEnabled(true);
