Fix CIS Unity build issues.

diff --git a/Engine/Plugins/Developer/Concert/ConcertUI/ConcertServerUI/Source/ConcertServerUI/Private/Widgets/SessionTabs/Live/LiveConcertSessionTab.h b/Engine/Plugins/Developer/Concert/ConcertUI/ConcertServerUI/Source/ConcertServerUI/Private/Widgets/SessionTabs/Live/LiveConcertSessionTab.h
--- a/Engine/Plugins/Developer/Concert/ConcertUI/ConcertServerUI/Source/ConcertServerUI/Private/Widgets/SessionTabs/Live/LiveConcertSessionTab.h
+++ b/Engine/Plugins/Developer/Concert/ConcertUI/ConcertServerUI/Source/ConcertServerUI/Private/Widgets/SessionTabs/Live/LiveConcertSessionTab.h
@@ -5,6 +5,7 @@
 #include "CoreMinimal.h"
 #include "ServerSessionHistoryController.h"
 #include "Widgets/SessionTabs/AbstractConcertSessionTab.h"
+#include "Types/SlateAttribute.h"
 
 class FConcertSessionPackageViewerController;
 class IConcertSyncServer;
diff --git a/Engine/Plugins/Developer/Concert/ConcertUI/ConcertSharedSlate/Source/ConcertSharedSlate/Public/Session/Activity/ActivityColumn.h b/Engine/Plugins/Developer/Concert/ConcertUI/ConcertSharedSlate/Source/ConcertSharedSlate/Public/Session/Activity/ActivityColumn.h
--- a/Engine/Plugins/Developer/Concert/ConcertUI/ConcertSharedSlate/Source/ConcertSharedSlate/Public/Session/Activity/ActivityColumn.h
+++ b/Engine/Plugins/Developer/Concert/ConcertUI/ConcertSharedSlate/Source/ConcertSharedSlate/Public/Session/Activity/ActivityColumn.h
@@ -3,6 +3,7 @@
 #pragma once
 
 #include "CoreMinimal.h"
+#include "Widgets/SOverlay.h"
 #include "Widgets/Views/SHeaderRow.h"
 
 class SConcertSessionActivities;
diff --git a/Engine/Plugins/Developer/Concert/ConcertUI/ConcertSharedSlate/Source/ConcertSharedSlate/Public/Session/Activity/SConcertSessionActivities.h b/Engine/Plugins/Developer/Concert/ConcertUI/ConcertSharedSlate/Source/ConcertSharedSlate/Public/Session/Activity/SConcertSessionActivities.h
--- a/Engine/Plugins/Developer/Concert/ConcertUI/ConcertSharedSlate/Source/ConcertSharedSlate/Public/Session/Activity/SConcertSessionActivities.h
+++ b/Engine/Plugins/Developer/Concert/ConcertUI/ConcertSharedSlate/Source/ConcertSharedSlate/Public/Session/Activity/SConcertSessionActivities.h
@@ -9,6 +9,7 @@
 #include "Widgets/DeclarativeSyntaxSupport.h"
 #include "Widgets/SCompoundWidget.h"
 #include "Widgets/Views/SListView.h"
+#include "Widgets/Layout/SExpandableArea.h"
 #include "ConcertSyncSessionTypes.h"
 
 struct FConcertSessionActivity;
