Crash when closing server after opening tab

diff --git a/Engine/Plugins/Developer/Concert/ConcertUI/ConcertSharedSlate/Source/ConcertSharedSlate/Private/Session/Activity/SConcertSessionActivities.cpp b/Engine/Plugins/Developer/Concert/ConcertUI/ConcertSharedSlate/Source/ConcertSharedSlate/Private/Session/Activity/SConcertSessionActivities.cpp
--- a/Engine/Plugins/Developer/Concert/ConcertUI/ConcertSharedSlate/Source/ConcertSharedSlate/Private/Session/Activity/SConcertSessionActivities.cpp
+++ b/Engine/Plugins/Developer/Concert/ConcertUI/ConcertSharedSlate/Source/ConcertSharedSlate/Private/Session/Activity/SConcertSessionActivities.cpp
@@ -215,7 +215,6 @@ void SConcertSessionActivities::Construct(const FArguments& InArgs)
 	GetPackageEventFn = InArgs._OnGetPackageEvent;
 	MakeColumnOverlayWidgetFn = InArgs._OnMakeColumnOverlayWidget;
 	
-	Columns = InArgs._Columns;
 	HighlightText = InArgs._HighlightText;
 	
 	TimeFormat = InArgs._TimeFormat;
@@ -256,7 +255,7 @@ void SConcertSessionActivities::Construct(const FArguments& InArgs)
 					.AllowOverscroll(EAllowOverscroll::No)
 					.OnListViewScrolled(this, &SConcertSessionActivities::OnListViewScrolled)
 					.OnSelectionChanged(this, &SConcertSessionActivities::OnListViewSelectionChanged)
-					.HeaderRow(CreateHeaderRow())
+					.HeaderRow(CreateHeaderRow(InArgs))
 				]
 			]
 
@@ -349,10 +348,11 @@ void SConcertSessionActivities::Construct(const FArguments& InArgs)
 	}
 }
 
-TSharedRef<SHeaderRow> SConcertSessionActivities::CreateHeaderRow()
+TSharedRef<SHeaderRow> SConcertSessionActivities::CreateHeaderRow(const FArguments& InArgs)
 {
 	using namespace UE::ConcertSharedSlate;
-
+	TArray<FActivityColumn> Columns = InArgs._Columns;
+	
 	// Before sort because columns can be placed before
 	const bool bContainsDateTime = Columns.FindByPredicate([](const FActivityColumn& Column) { return Column.ColumnId ==  ActivityColumn::DateTimeColumnId; }) != nullptr; 
 	checkf(!bContainsDateTime, TEXT("DateTime is already created by SConcertSessionActivities!"))
@@ -365,13 +365,15 @@ TSharedRef<SHeaderRow> SConcertSessionActivities::CreateHeaderRow()
 	checkf(!bContainsSummary, TEXT("Summary is already created by SConcertSessionActivities!"))
 	Columns.Add(ActivityColumn::Summary());
 	
-	const TSharedRef<SHeaderRow> HeaderRow = SNew(SHeaderRow);
+	HeaderRow = SNew(SHeaderRow);
 	for (FActivityColumn& Column : Columns)
 	{
-		HeaderRow->AddColumn(Column);
+		// SHeaderRow owns the columns and deletes them when destroyed
+		FActivityColumn* ManagedByHeaderRow = new FActivityColumn(MoveTemp(Column));
+		HeaderRow->AddColumn(*ManagedByHeaderRow);
 	}
 	
-	return HeaderRow;
+	return HeaderRow.ToSharedRef();
 }
 
 TSharedRef<ITableRow> SConcertSessionActivities::OnGenerateActivityRowWidget(TSharedPtr<FConcertSessionActivity> Activity, const TSharedRef<STableViewBase>& OwnerTable)
@@ -379,11 +381,11 @@ TSharedRef<ITableRow> SConcertSessionActivities::OnGenerateActivityRowWidget(TSh
 	const TOptional<FConcertClientInfo> ActivityClient = GetActivityUserFn.IsBound() ? GetActivityUserFn.Execute(Activity->Activity.EndpointId) : TOptional<FConcertClientInfo>{};
 	const SConcertSessionActivityRow::FGetColumn ColumnGetter = SConcertSessionActivityRow::FGetColumn::CreateLambda([this](const FName& ColumnId) -> const FActivityColumn*
 	{
-		for (const FActivityColumn& Column : Columns)
+		for (const SHeaderRow::FColumn& Column : HeaderRow->GetColumns())
 		{
 			if (Column.ColumnId == ColumnId)
 			{
-				return &Column;
+				return static_cast<const FActivityColumn*>(&Column);
 			}
 		}
 
@@ -731,9 +733,10 @@ FText SConcertSessionActivities::UpdateTextFilter(const FText& InFilterText)
 
 void SConcertSessionActivities::PopulateSearchStrings(const FConcertSessionActivity& Activity, TArray<FString>& OutSearchStrings)
 {
-	for (const FActivityColumn& Column : Columns)
+	for (const SHeaderRow::FColumn& Column : HeaderRow->GetColumns())
 	{
-		Column.ExecutePopulateSearchString(SharedThis(this), Activity, OutSearchStrings);
+		const FActivityColumn& CastColumn = static_cast<const FActivityColumn&>(Column);
+		CastColumn.ExecutePopulateSearchString(SharedThis(this), Activity, OutSearchStrings);
 	}
 }
 
diff --git a/Engine/Plugins/Developer/Concert/ConcertUI/ConcertSharedSlate/Source/ConcertSharedSlate/Public/Session/Activity/SConcertSessionActivities.h b/Engine/Plugins/Developer/Concert/ConcertUI/ConcertSharedSlate/Source/ConcertSharedSlate/Public/Session/Activity/SConcertSessionActivities.h
--- a/Engine/Plugins/Developer/Concert/ConcertUI/ConcertSharedSlate/Source/ConcertSharedSlate/Public/Session/Activity/SConcertSessionActivities.h
+++ b/Engine/Plugins/Developer/Concert/ConcertUI/ConcertSharedSlate/Source/ConcertSharedSlate/Public/Session/Activity/SConcertSessionActivities.h
@@ -16,10 +16,11 @@ struct FConcertSessionActivity;
 struct FConcertClientInfo;
 struct FConcertPackageInfo;
 struct FConcertTransactionEventBase;
-class SUndoHistoryDetails;
+class SExpendableArea;
+class SHeaderRow;
 class SPackageDetails;
+class SUndoHistoryDetails;
 class STextBlock;
-class SExpendableArea;
 
 /** Filters for the concert session activity view. */
 enum class EConcertActivityFilterFlags
@@ -122,7 +123,7 @@ public:
 		/** Show/hide a message overlay above the activities list explaining why no activities are displayed. */
 		SLATE_ATTRIBUTE(FText, NoActivitiesReasonText)
 	SLATE_END_ARGS();
-
+	
 	/**
 	 * Construct the recovery widget.
 	 * @param InArgs The widgets arguments and attributes.
@@ -169,7 +170,7 @@ public:
 private:
 
 	// Widget creation
-	TSharedRef<SHeaderRow> CreateHeaderRow();
+	TSharedRef<SHeaderRow> CreateHeaderRow(const FArguments& InArgs);
 	TSharedRef<ITableRow> OnGenerateActivityRowWidget(TSharedPtr<FConcertSessionActivity> Item, const TSharedRef<STableViewBase>& OwnerTable);
 	void FetchActivities();
 
@@ -199,9 +200,6 @@ private:
 	void OnPostTick(float);
 
 private:
-
-	/** The columns that are being displayed */
-	TArray<FActivityColumn> Columns;
 	
 	/** List of all activities (including the filtered out ones). */
 	TArray<TSharedPtr<FConcertSessionActivity>> AllActivities;
@@ -268,10 +266,12 @@ private:
 
 	/** True once the activity provider function (FetchActivitiesFn) returns true. */
 	bool bAllActivitiesFetched = false;
-
+	
 	/** Utility class used to tokenize and match text displayed in the list view. */
 	TSharedPtr<TTextFilter<const FConcertSessionActivity&>> SearchTextFilter;
 
+	TSharedPtr<SHeaderRow> HeaderRow;
+	
 	/** The expandable area under which the activity details are displayed. */
 	TSharedPtr<SExpandableArea> ExpandableDetails;
 
