Add project and version column to the browser panel.

diff --git a/Engine/Plugins/Developer/Concert/ConcertApp/MultiUserClient/Source/MultiUserClient/Private/Widgets/SConcertBrowser.cpp b/Engine/Plugins/Developer/Concert/ConcertApp/MultiUserClient/Source/MultiUserClient/Private/Widgets/SConcertBrowser.cpp
--- a/Engine/Plugins/Developer/Concert/ConcertApp/MultiUserClient/Source/MultiUserClient/Private/Widgets/SConcertBrowser.cpp
+++ b/Engine/Plugins/Developer/Concert/ConcertApp/MultiUserClient/Source/MultiUserClient/Private/Widgets/SConcertBrowser.cpp
@@ -135,7 +135,7 @@ public:
 	virtual TOptional<FConcertSessionInfo> GetActiveSessionInfo(const FGuid& AdminEndpoint, const FGuid& SessionId) const override;
 	virtual TOptional<FConcertSessionInfo> GetArchivedSessionInfo(const FGuid& AdminEndpoint, const FGuid& SessionId) const override;
 
-	virtual void CreateSession(const FGuid& ServerAdminEndpointId, const FString& SessionName) override;
+	virtual void CreateSession(const FGuid& ServerAdminEndpointId, const FString& SessionName, const FString& ProjectName) override;
 	virtual void ArchiveSession(const FGuid& ServerAdminEndpointId, const FGuid& SessionId, const FString& ArchiveName, const FConcertSessionFilter& SessionFilter) override;
 	virtual void RestoreSession(const FGuid& ServerAdminEndpointId, const FGuid& SessionId, const FString& RestoredName, const FConcertSessionFilter& SessionFilter) override;
 	virtual void RenameActiveSession(const FGuid& ServerAdminEndpointId, const FGuid& SessionId, const FString& NewName) override;
@@ -296,7 +296,7 @@ TOptional<FConcertSessionInfo> FConcertClientSessionBrowserController::GetArchiv
 	return SessionInfo != nullptr ? (*SessionInfo)->SessionInfo : TOptional<FConcertSessionInfo>{};
 }
 
-void FConcertClientSessionBrowserController::CreateSession(const FGuid& ServerAdminEndpointId, const FString& SessionName)
+void FConcertClientSessionBrowserController::CreateSession(const FGuid& ServerAdminEndpointId, const FString& SessionName, const FString& ProjectName)
 {
 	FAsyncRequest& CreateRequest = CreateSessionRequests.AddDefaulted_GetRef();
 	TWeakPtr<uint8, ESPMode::ThreadSafe> CreateRequestExecutionToken = CreateRequest.ResetExecutionToken();
@@ -305,6 +305,7 @@ void FConcertClientSessionBrowserController::CreateSession(const FGuid& ServerAd
 	// On failure: An async notification banner will be displayer to the user.
 	FConcertCreateSessionArgs CreateSessionArgs;
 	CreateSessionArgs.SessionName = SessionName;
+	// Project Name is not used an override for the client session create.
 	ConcertClient->CreateSession(ServerAdminEndpointId, CreateSessionArgs).Next([this, CreateRequestExecutionToken, ServerAdminEndpointId, SessionName](EConcertResponseCode ResponseCode)
 	{
 		if (TSharedPtr<uint8, ESPMode::ThreadSafe> ExecutionToken = CreateRequestExecutionToken.Pin())
@@ -1647,18 +1648,10 @@ void SConcertClientSessionBrowser::PopulateSessionInfoGrid(SGridPanel& Grid, con
 	AddDetailRow(Grid, Row++, LOCTEXT("SessionName", "Session Name:"), FText::FromString(SessionInfo.SessionName));
 	AddDetailRow(Grid, Row++, LOCTEXT("Owner", "Owner:"), FText::FromString(SessionInfo.OwnerUserName));
 	AddDetailRow(Grid, Row++, LOCTEXT("Project", "Project:"), FText::FromString(SessionInfo.Settings.ProjectName));
-	//AddDetailRow(Grid, Row++, LOCTEXT("BaseVersion", "Base Version:"), FText::AsNumber(SessionInfo.Settings.BaseRevision, &FNumberFormattingOptions::DefaultNoGrouping()));
 	if (SessionInfo.VersionInfos.Num() > 0)
 	{
 		const FConcertSessionVersionInfo& VersionInfo = SessionInfo.VersionInfos.Last();
-		AddDetailRow(Grid, Row++, LOCTEXT("EngineVersion", "Engine Version:"),
-			FText::Format(
-				LOCTEXT("EngineVersionFmt", "{0}.{1}.{2}-{3}"),
-				FText::AsNumber(VersionInfo.EngineVersion.Major, &FNumberFormattingOptions::DefaultNoGrouping()),
-				FText::AsNumber(VersionInfo.EngineVersion.Minor, &FNumberFormattingOptions::DefaultNoGrouping()),
-				FText::AsNumber(VersionInfo.EngineVersion.Patch, &FNumberFormattingOptions::DefaultNoGrouping()),
-				FText::AsNumber(VersionInfo.EngineVersion.Changelist, &FNumberFormattingOptions::DefaultNoGrouping())
-			));
+		AddDetailRow(Grid, Row++, LOCTEXT("EngineVersion", "Engine Version:"), VersionInfo.AsText());
 	}
 	AddDetailRow(Grid, Row++, LOCTEXT("ServerEndPointId", "Server Endpoint ID:"), FText::FromString(SessionInfo.ServerEndpointId.ToString()));
 }
diff --git a/Engine/Plugins/Developer/Concert/ConcertApp/MultiUserServer/Source/MultiUserServer/Private/Widgets/Browser/ConcertServerSessionBrowserController.cpp b/Engine/Plugins/Developer/Concert/ConcertApp/MultiUserServer/Source/MultiUserServer/Private/Widgets/Browser/ConcertServerSessionBrowserController.cpp
--- a/Engine/Plugins/Developer/Concert/ConcertApp/MultiUserServer/Source/MultiUserServer/Private/Widgets/Browser/ConcertServerSessionBrowserController.cpp
+++ b/Engine/Plugins/Developer/Concert/ConcertApp/MultiUserServer/Source/MultiUserServer/Private/Widgets/Browser/ConcertServerSessionBrowserController.cpp
@@ -98,18 +98,22 @@ TOptional<FConcertSessionInfo> FConcertServerSessionBrowserController::GetArchiv
 	return ServerInstance->GetConcertServer()->GetArchivedSessionInfo(SessionId);
 }
 
-void FConcertServerSessionBrowserController::CreateSession(const FGuid& ServerAdminEndpointId, const FString& SessionName)
+void FConcertServerSessionBrowserController::CreateSession(const FGuid& ServerAdminEndpointId, const FString& SessionName, const FString& ProjectName)
 {
 	FConcertSessionInfo SessionInfo = ServerInstance->GetConcertServer()->CreateSessionInfo();
 	SessionInfo.SessionName = SessionName;
 	SessionInfo.Settings.Initialize();
+	if (!ProjectName.IsEmpty())
+	{
+		SessionInfo.Settings.ProjectName = ProjectName;
+	}
 	FConcertSessionVersionInfo VersionInfo;
 	VersionInfo.Initialize();
 	SessionInfo.VersionInfos.Emplace(VersionInfo);
-	
+
 	FText FailureReason = FText::GetEmpty();
 	const bool bSuccess = ServerInstance->GetConcertServer()->CreateSession(SessionInfo, FailureReason) != nullptr;
-	
+
 	NotifyUserOfFinishedSessionAction(bSuccess,
 		bSuccess ? FText::Format(LOCTEXT("CreatedSessionFmt", "Created Session '{0}'"), FText::FromString(SessionName)) : FText::Format(LOCTEXT("FailedToCreateSessionFmt", "Failed to create Session '{0}'"), FText::FromString(SessionName)),
 		FailureReason
diff --git a/Engine/Plugins/Developer/Concert/ConcertApp/MultiUserServer/Source/MultiUserServer/Private/Widgets/Browser/ConcertServerSessionBrowserController.h b/Engine/Plugins/Developer/Concert/ConcertApp/MultiUserServer/Source/MultiUserServer/Private/Widgets/Browser/ConcertServerSessionBrowserController.h
--- a/Engine/Plugins/Developer/Concert/ConcertApp/MultiUserServer/Source/MultiUserServer/Private/Widgets/Browser/ConcertServerSessionBrowserController.h
+++ b/Engine/Plugins/Developer/Concert/ConcertApp/MultiUserServer/Source/MultiUserServer/Private/Widgets/Browser/ConcertServerSessionBrowserController.h
@@ -25,26 +25,26 @@ public:
 	int32 GetNumConnectedClients(const FGuid& SessionId) const;
 
 	~FConcertServerSessionBrowserController();
-	
+
 	//~ Begin IConcertComponent Interface
 	virtual void Init(const FConcertComponentInitParams& Params) override;
 	//~ End IConcertComponent Interface
-	
+
 	//~ Begin IConcertSessionBrowserController Interface
 	virtual TArray<FConcertServerInfo> GetServers() const override;
 	virtual TArray<FActiveSessionInfo> GetActiveSessions() const override;
 	virtual TArray<FArchivedSessionInfo> GetArchivedSessions() const override;
 	virtual TOptional<FConcertSessionInfo> GetActiveSessionInfo(const FGuid& AdminEndpoint, const FGuid& SessionId) const override;
 	virtual TOptional<FConcertSessionInfo> GetArchivedSessionInfo(const FGuid& AdminEndpoint, const FGuid& SessionId) const override;
-	
-	virtual void CreateSession(const FGuid& ServerAdminEndpointId, const FString& SessionName) override;
+
+	virtual void CreateSession(const FGuid& ServerAdminEndpointId, const FString& SessionName, const FString& ProjectName) override;
 	virtual void ArchiveSession(const FGuid& ServerAdminEndpointId, const FGuid& SessionId, const FString& ArchiveName, const FConcertSessionFilter& SessionFilter) override;
 	virtual void RestoreSession(const FGuid& ServerAdminEndpointId, const FGuid& SessionId, const FString& RestoredName, const FConcertSessionFilter& SessionFilter) override;
 	virtual void RenameActiveSession(const FGuid& ServerAdminEndpointId, const FGuid& SessionId, const FString& NewName) override { RenameSession(ServerAdminEndpointId, SessionId, NewName); }
 	virtual void RenameArchivedSession(const FGuid& ServerAdminEndpointId, const FGuid& SessionId, const FString& NewName) override { RenameSession(ServerAdminEndpointId, SessionId, NewName); }
 	virtual void DeleteActiveSession(const FGuid& ServerAdminEndpointId, const FGuid& SessionId) override { DeleteSession(ServerAdminEndpointId, SessionId); }
 	virtual void DeleteArchivedSession(const FGuid& ServerAdminEndpointId, const FGuid& SessionId) override { DeleteSession(ServerAdminEndpointId, SessionId); }
-	
+
 	// The server operator always has permission for these actions:
 	virtual bool CanRenameActiveSession(const FGuid& ServerAdminEndpointId, const FGuid& SessionId) const override { return true; }
 	virtual bool CanRenameArchivedSession(const FGuid& ServerAdminEndpointId, const FGuid& SessionId) const override { return true; }
@@ -58,9 +58,9 @@ private:
 	TSharedPtr<IConcertSyncServer> ServerInstance;
 	/** Used to open selected sessions */
 	TWeakPtr<FConcertServerWindowController> Owner;
-	
+
 	TSharedPtr<SConcertServerSessionBrowser> ConcertBrowser;
-	
+
 	TSharedRef<SDockTab> SpawnSessionBrowserTab(const FSpawnTabArgs& Args);
 
 	// Update view when session list changes
@@ -68,7 +68,7 @@ private:
 	void OnLiveSessionDestroyed(const IConcertServer&, TSharedRef<IConcertServerSession>) { RefreshSessionList(); }
 	void OnArchivedSessionCreated(bool, const IConcertServer&, const FString&, const FConcertSessionInfo&) { RefreshSessionList(); }
 	void OnArchivedSessionDestroyed(const IConcertServer&, const FGuid&) { RefreshSessionList(); }
-	
+
 	void RefreshSessionList();
 
 	// Session actions
diff --git a/Engine/Plugins/Developer/Concert/ConcertMain/Source/Concert/Private/ConcertVersion.cpp b/Engine/Plugins/Developer/Concert/ConcertMain/Source/Concert/Private/ConcertVersion.cpp
--- a/Engine/Plugins/Developer/Concert/ConcertMain/Source/Concert/Private/ConcertVersion.cpp
+++ b/Engine/Plugins/Developer/Concert/ConcertMain/Source/Concert/Private/ConcertVersion.cpp
@@ -191,4 +191,14 @@ bool FConcertSessionVersionInfo::Validate(const FConcertSessionVersionInfo& InOt
 	return true;
 }
 
+FText FConcertSessionVersionInfo::AsText() const
+{
+	return FText::Format(
+		LOCTEXT("EngineVersionFmt", "{0}.{1}.{2}-{3}"),
+		FText::AsNumber(EngineVersion.Major, &FNumberFormattingOptions::DefaultNoGrouping()),
+		FText::AsNumber(EngineVersion.Minor, &FNumberFormattingOptions::DefaultNoGrouping()),
+		FText::AsNumber(EngineVersion.Patch, &FNumberFormattingOptions::DefaultNoGrouping()),
+		FText::AsNumber(EngineVersion.Changelist, &FNumberFormattingOptions::DefaultNoGrouping())
+		);
+}
 #undef LOCTEXT_NAMESPACE
diff --git a/Engine/Plugins/Developer/Concert/ConcertMain/Source/Concert/Public/ConcertVersion.h b/Engine/Plugins/Developer/Concert/ConcertMain/Source/Concert/Public/ConcertVersion.h
--- a/Engine/Plugins/Developer/Concert/ConcertMain/Source/Concert/Public/ConcertVersion.h
+++ b/Engine/Plugins/Developer/Concert/ConcertMain/Source/Concert/Public/ConcertVersion.h
@@ -108,6 +108,9 @@ struct FConcertSessionVersionInfo
 	/** Validate this version info against another */
 	CONCERT_API bool Validate(const FConcertSessionVersionInfo& InOther, const EConcertVersionValidationMode InValidationMode, FText* OutFailureReason = nullptr) const;
 
+	/** Return the version information in string represenation. */
+	CONCERT_API FText AsText() const;
+
 	/** File version info */
 	UPROPERTY()
 	FConcertFileVersionInfo FileVersion;
diff --git a/Engine/Plugins/Developer/Concert/ConcertUI/ConcertSharedSlate/Source/ConcertSharedSlate/Private/Session/Browser/SConcertSessionBrowser.cpp b/Engine/Plugins/Developer/Concert/ConcertUI/ConcertSharedSlate/Source/ConcertSharedSlate/Private/Session/Browser/SConcertSessionBrowser.cpp
--- a/Engine/Plugins/Developer/Concert/ConcertUI/ConcertSharedSlate/Source/ConcertSharedSlate/Private/Session/Browser/SConcertSessionBrowser.cpp
+++ b/Engine/Plugins/Developer/Concert/ConcertUI/ConcertSharedSlate/Source/ConcertSharedSlate/Private/Session/Browser/SConcertSessionBrowser.cpp
@@ -2,6 +2,7 @@
 
 #include "Session/Browser/SConcertSessionBrowser.h"
 
+#include "Algo/ForEach.h"
 #include "ConcertFrontendStyle.h"
 #include "Session/Browser/ConcertBrowserUtils.h"
 #include "Session/Browser/ConcertSessionBrowserSettings.h"
@@ -111,6 +112,41 @@ TSharedRef<SWidget> SConcertSessionBrowser::MakeBrowserContent(const FArguments&
 		];
 }
 
+namespace UE::SessionBrowser::Private
+{
+
+template <typename SessionType>
+constexpr FConcertSessionItem::EType GetEType()
+{
+	if constexpr(std::is_same<IConcertSessionBrowserController::FActiveSessionInfo,SessionType>::value)
+	{
+		return FConcertSessionItem::EType::ActiveSession;
+	}
+	return FConcertSessionItem::EType::ArchivedSession;
+}
+
+template <typename SessionType>
+static FConcertSessionItem GenerateSessionItem(const SessionType& Session)
+{
+	FText Version = LOCTEXT("ConcertEmptyVersion", "Unset");
+	if (Session.SessionInfo.VersionInfos.Num()>0)
+	{
+		const FConcertSessionVersionInfo& VersionInfo = Session.SessionInfo.VersionInfos.Last();
+		Version = VersionInfo.AsText();
+	}
+	return FConcertSessionItem{
+		GetEType<SessionType>(),
+		Session.ServerInfo.AdminEndpointId,
+		Session.SessionInfo.SessionId,
+		Session.SessionInfo.SessionName,
+		Session.ServerInfo.ServerName,
+		Session.SessionInfo.Settings.ProjectName,
+		Version.ToString(),
+		Session.ServerInfo.ServerFlags,{}};
+}
+
+}
+
 void SConcertSessionBrowser::RefreshSessionList()
 {
 	// Remember the selected instances (if any).
@@ -123,7 +159,6 @@ void SConcertSessionBrowser::RefreshSessionList()
 	{
 		return SelectedItems.ContainsByPredicate([Item](const TSharedPtr<FConcertSessionItem>& Visited) { return *Visited == Item; });
 	};
-
 	// Matches the object instances before the update to the new instance after the update.
 	auto ReconcileObjectInstances = [&](TSharedPtr<FConcertSessionItem> NewItem)
 	{
@@ -136,31 +171,22 @@ void SConcertSessionBrowser::RefreshSessionList()
 			NewEditableRowParent = NewItem;
 		}
 	};
-
-	// Clear sessions.
-	Sessions.Reset();
-
-	// Populate the live sessions.
-	for (const IConcertSessionBrowserController::FActiveSessionInfo& ActiveSession : GetController()->GetActiveSessions())
+	auto ConvertSession = [&](const auto& Session)
 	{
-		FConcertSessionItem NewItem(FConcertSessionItem::EType::ActiveSession, ActiveSession.SessionInfo.SessionName, ActiveSession.SessionInfo.SessionId, ActiveSession.ServerInfo.ServerName, ActiveSession.ServerInfo.AdminEndpointId, ActiveSession.ServerInfo.ServerFlags);
+		FConcertSessionItem NewItem = UE::SessionBrowser::Private::GenerateSessionItem(Session);
 		if (!IsFilteredOut(NewItem))
 		{
-			Sessions.Emplace(MakeShared<FConcertSessionItem>(MoveTemp(NewItem)));
+			Sessions.Emplace(
+				MakeShared<FConcertSessionItem>(MoveTemp(NewItem)));
 			ReconcileObjectInstances(Sessions.Last());
 		}
-	}
+	};
+	Sessions.Reset();
 
-	// Populate the archived.
-	for (const IConcertSessionBrowserController::FArchivedSessionInfo& ArchivedSession : GetController()->GetArchivedSessions())
-	{
-		FConcertSessionItem NewItem(FConcertSessionItem::EType::ArchivedSession, ArchivedSession.SessionInfo.SessionName, ArchivedSession.SessionInfo.SessionId, ArchivedSession.ServerInfo.ServerName, ArchivedSession.ServerInfo.AdminEndpointId, ArchivedSession.ServerInfo.ServerFlags);
-		if (!IsFilteredOut(NewItem))
-		{
-			Sessions.Emplace(MakeShared<FConcertSessionItem>(MoveTemp(NewItem)));
-			ReconcileObjectInstances(Sessions.Last());
-		}
-	}
+	TArray<IConcertSessionBrowserController::FActiveSessionInfo> ActiveSessions = GetController()->GetActiveSessions();
+	TArray<IConcertSessionBrowserController::FArchivedSessionInfo> ArchivedSessions = GetController()->GetArchivedSessions();
+	Algo::ForEach(ActiveSessions, ConvertSession);
+	Algo::ForEach(ArchivedSessions, ConvertSession);
 
 	// Restore the editable row state. (SortSessionList() below will ensure the parent/child relationship)
 	EditableSessionRowParent = NewEditableRowParent;
@@ -224,9 +250,9 @@ bool SConcertSessionBrowser::IsFilteredOut(const FConcertSessionItem& Item) cons
 	bool bIsDefaultServer = LastDefaultServerUrl.IsEmpty() || Item.ServerName == LastDefaultServerUrl;
 
 	return (!PersistentSettings->bShowActiveSessions && (Item.Type == FConcertSessionItem::EType::ActiveSession || Item.Type == FConcertSessionItem::EType::SaveSession)) ||
-	       (!PersistentSettings->bShowArchivedSessions && (Item.Type == FConcertSessionItem::EType::ArchivedSession || Item.Type == FConcertSessionItem::EType::RestoreSession)) ||
-	       (PersistentSettings->bShowDefaultServerSessionsOnly && !bIsDefaultServer) ||
-	       !SearchTextFilter->PassesFilter(Item);
+		   (!PersistentSettings->bShowArchivedSessions && (Item.Type == FConcertSessionItem::EType::ArchivedSession || Item.Type == FConcertSessionItem::EType::RestoreSession)) ||
+		   (PersistentSettings->bShowDefaultServerSessionsOnly && !bIsDefaultServer) ||
+		   !SearchTextFilter->PassesFilter(Item);
 }
 
 FText SConcertSessionBrowser::HighlightSearchText() const
@@ -347,6 +373,18 @@ TSharedRef<SWidget> SConcertSessionBrowser::MakeSessionTableView()
 			.DefaultLabel(LOCTEXT("Server", "Server"))
 			.SortPriority(this, &SConcertSessionBrowser::GetColumnSortPriority, ConcertBrowserUtils::ServerColName)
 			.SortMode(this, &SConcertSessionBrowser::GetColumnSortMode, ConcertBrowserUtils::ServerColName)
+			.OnSort(this, &SConcertSessionBrowser::OnColumnSortModeChanged)
+
+			+SHeaderRow::Column(ConcertBrowserUtils::ProjectColName)
+			.DefaultLabel(LOCTEXT("Project", "Project"))
+			.SortPriority(this, &SConcertSessionBrowser::GetColumnSortPriority, ConcertBrowserUtils::ProjectColName)
+			.SortMode(this, &SConcertSessionBrowser::GetColumnSortMode, ConcertBrowserUtils::ProjectColName)
+			.OnSort(this, &SConcertSessionBrowser::OnColumnSortModeChanged)
+
+			+SHeaderRow::Column(ConcertBrowserUtils::VersionColName)
+			.DefaultLabel(LOCTEXT("Version", "Version"))
+			.SortPriority(this, &SConcertSessionBrowser::GetColumnSortPriority, ConcertBrowserUtils::VersionColName)
+			.SortMode(this, &SConcertSessionBrowser::GetColumnSortMode, ConcertBrowserUtils::VersionColName)
 			.OnSort(this, &SConcertSessionBrowser::OnColumnSortModeChanged));
 }
 
@@ -559,19 +597,20 @@ TSharedRef<ITableRow> SConcertSessionBrowser::MakeRestoreSessionRowWidget(const
 void SConcertSessionBrowser::InsertNewSessionEditableRowInternal()
 {
 	// Insert a 'new session' editable row.
-	InsertEditableSessionRow(MakeShared<FConcertSessionItem>(FConcertSessionItem::EType::NewSession, FString(), FGuid(), FString(), FGuid(), EConcertServerFlags::None), nullptr);
+	FConcertSessionItem Item{FConcertSessionItem::EType::NewSession};
+	InsertEditableSessionRow(MakeShared<FConcertSessionItem>(MoveTemp(Item)), nullptr);
 }
 
 void SConcertSessionBrowser::InsertRestoreSessionAsEditableRowInternal(const TSharedPtr<FConcertSessionItem>& ArchivedItem)
 {
 	// Insert the 'restore session as ' editable row just below the 'archived' item to restore.
-	InsertEditableSessionRow(MakeShared<FConcertSessionItem>(FConcertSessionItem::EType::RestoreSession, ArchivedItem->SessionName, ArchivedItem->SessionId, ArchivedItem->ServerName, ArchivedItem->ServerAdminEndpointId, ArchivedItem->ServerFlags), ArchivedItem);
+	InsertEditableSessionRow(MakeShared<FConcertSessionItem>(ArchivedItem->MakeCopyAsType(FConcertSessionItem::EType::RestoreSession)), ArchivedItem);
 }
 
 void SConcertSessionBrowser::InsertArchiveSessionAsEditableRow(const TSharedPtr<FConcertSessionItem>& LiveItem)
 {
 	// Insert the 'save session as' editable row just below the 'active' item to save.
-	InsertEditableSessionRow(MakeShared<FConcertSessionItem>(FConcertSessionItem::EType::SaveSession, LiveItem->SessionName, LiveItem->SessionId, LiveItem->ServerName, LiveItem->ServerAdminEndpointId, LiveItem->ServerFlags), LiveItem);
+	InsertEditableSessionRow(MakeShared<FConcertSessionItem>(LiveItem->MakeCopyAsType(FConcertSessionItem::EType::SaveSession)), LiveItem);
 }
 
 void SConcertSessionBrowser::InsertEditableSessionRow(TSharedPtr<FConcertSessionItem> EditableItem, TSharedPtr<FConcertSessionItem> ParentItem)
@@ -991,7 +1030,7 @@ FReply SConcertSessionBrowser::OnKeyDown(const FGeometry& MyGeometry, const FKey
 
 void SConcertSessionBrowser::RequestCreateSession(const TSharedPtr<FConcertSessionItem>& NewItem)
 {
-	GetController()->CreateSession(NewItem->ServerAdminEndpointId, NewItem->SessionName);
+	GetController()->CreateSession(NewItem->ServerAdminEndpointId, NewItem->SessionName, NewItem->ProjectName);
 	RemoveSessionRow(NewItem); // The row used to edit the session name and pick the server.
 }
 
@@ -1043,4 +1082,4 @@ void SConcertSessionBrowser::RequestDeleteSession(const TSharedPtr<FConcertSessi
 	OnRequestedDeleteSession.ExecuteIfBound(DeletedItem);
 }
 
-#undef LOCTEXT_NAMESPACE
\ No newline at end of file
+#undef LOCTEXT_NAMESPACE
diff --git a/Engine/Plugins/Developer/Concert/ConcertUI/ConcertSharedSlate/Source/ConcertSharedSlate/Private/Session/Browser/SNewSessionRow.cpp b/Engine/Plugins/Developer/Concert/ConcertUI/ConcertSharedSlate/Source/ConcertSharedSlate/Private/Session/Browser/SNewSessionRow.cpp
--- a/Engine/Plugins/Developer/Concert/ConcertUI/ConcertSharedSlate/Source/ConcertSharedSlate/Private/Session/Browser/SNewSessionRow.cpp
+++ b/Engine/Plugins/Developer/Concert/ConcertUI/ConcertSharedSlate/Source/ConcertSharedSlate/Private/Session/Browser/SNewSessionRow.cpp
@@ -79,7 +79,7 @@ TSharedRef<SWidget> SNewSessionRow::GenerateWidgetForColumn(const FName& ColumnN
 				.OnTextChanged(this, &SNewSessionRow::OnSessionNameChanged)
 			];
 	}
-	else
+	else if (ColumnName == ConcertBrowserUtils::ServerColName)
 	{
 		return SNew(SHorizontalBox)
 
@@ -95,8 +95,11 @@ TSharedRef<SWidget> SNewSessionRow::GenerateWidgetForColumn(const FName& ColumnN
 				[
 					MakeSelectedServerWidget()
 				]
-			]
-
+			];
+	}
+	else if (ColumnName == ConcertBrowserUtils::VersionColName)
+	{
+		return SNew(SHorizontalBox)
 			+SHorizontalBox::Slot()
 			.AutoWidth()
 			.Padding(2.0f)
@@ -126,6 +129,21 @@ TSharedRef<SWidget> SNewSessionRow::GenerateWidgetForColumn(const FName& ColumnN
 				]
 			];
 	}
+	else
+	{
+		check(ColumnName == ConcertBrowserUtils::ProjectColName);
+		FText ProjectName = FText::AsCultureInvariant(FApp::GetProjectName());
+		return SNew(SBox)
+			.VAlign(VAlign_Center)
+			.Padding(FMargin(0, 0, 2, 0))
+			[
+				SAssignNew(ProjectNameTextBox, SEditableTextBox)
+				.HintText(ProjectName)
+				.Text(ProjectName)
+				.OnTextCommitted(this, &SNewSessionRow::OnProjectNameCommitted)
+				.OnKeyDownHandler(this, &SNewSessionRow::OnKeyDownHandler)
+			];
+	}
 }
 
 TSharedRef<SWidget> SNewSessionRow::OnGenerateServersComboOptionWidget(TSharedPtr<FConcertServerInfo> ServerItem)
@@ -302,13 +320,19 @@ FReply SNewSessionRow::OnAccept()
 {
 	TSharedPtr<FConcertSessionItem> ItemPin = Item.Pin();
 	FString NewSessionName = EditableSessionName->GetText().ToString();
-
 	FText InvalidNameErrorMsg = ConcertSettingsUtils::ValidateSessionName(NewSessionName);
+
 	if (InvalidNameErrorMsg.IsEmpty())
 	{
-		ItemPin->SessionName = EditableSessionName->GetText().ToString();
+		ItemPin->SessionName = NewSessionName;
 		ItemPin->ServerName = ServersComboBox->GetSelectedItem()->ServerName;
 		ItemPin->ServerAdminEndpointId = ServersComboBox->GetSelectedItem()->AdminEndpointId;
+		FString CandidateProjectName = ProjectNameTextBox->GetText().ToString();
+		if (CandidateProjectName.IsEmpty())
+		{
+			CandidateProjectName = FApp::GetProjectName();
+		}
+		ItemPin->ProjectName = CandidateProjectName;
 		AcceptFunc(ItemPin); // Delegate to create the session.
 	}
 	else
@@ -331,6 +355,14 @@ void SNewSessionRow::OnSessionNameChanged(const FText& NewName)
 	EditableSessionName->SetError(ConcertSettingsUtils::ValidateSessionName(NewName.ToString()));
 }
 
+void SNewSessionRow::OnProjectNameCommitted(const FText &NewProjectName, ETextCommit::Type CommitType)
+{
+	if (CommitType == ETextCommit::OnEnter)
+	{
+		OnAccept(); // Create the session;.
+	}
+}
+
 void SNewSessionRow::OnSessionNameCommitted(const FText& NewText, ETextCommit::Type CommitType)
 {
 	if (CommitType == ETextCommit::OnEnter)
@@ -345,4 +377,4 @@ FReply SNewSessionRow::OnKeyDownHandler(const FGeometry&, const FKeyEvent& KeyEv
 	return KeyEvent.GetKey() == EKeys::Escape ? OnDecline() : FReply::Unhandled();
 }
 
-#undef LOCTEXT_NAMESPACE
\ No newline at end of file
+#undef LOCTEXT_NAMESPACE
diff --git a/Engine/Plugins/Developer/Concert/ConcertUI/ConcertSharedSlate/Source/ConcertSharedSlate/Private/Session/Browser/SNewSessionRow.h b/Engine/Plugins/Developer/Concert/ConcertUI/ConcertSharedSlate/Source/ConcertSharedSlate/Private/Session/Browser/SNewSessionRow.h
--- a/Engine/Plugins/Developer/Concert/ConcertUI/ConcertSharedSlate/Source/ConcertSharedSlate/Private/Session/Browser/SNewSessionRow.h
+++ b/Engine/Plugins/Developer/Concert/ConcertUI/ConcertSharedSlate/Source/ConcertSharedSlate/Private/Session/Browser/SNewSessionRow.h
@@ -35,19 +35,19 @@ public:
 	SLATE_ARGUMENT(FGetServersFunc, GetServerFunc)
 	SLATE_ARGUMENT(FAcceptFunc, OnAcceptFunc)
 	SLATE_ARGUMENT(FDeclineFunc, OnDeclineFunc)
-	
+
 	SLATE_ATTRIBUTE(FText, HighlightText)
 	SLATE_ATTRIBUTE(FString, DefaultServerURL)
 
 	SLATE_END_ARGS()
 
 	void Construct(const FArguments& InArgs, TSharedPtr<FConcertSessionItem> InItem, const TSharedRef<STableViewBase>& InOwnerTableView);
-	
+
 	virtual TSharedRef<SWidget> GenerateWidgetForColumn(const FName& ColumnName) override;
 	virtual void Tick(const FGeometry& AllottedGeometry, const double InCurrentTime, const float InDeltaTime) override;
 
 private:
-	
+
 	TSharedRef<SWidget> OnGenerateServersComboOptionWidget(TSharedPtr<FConcertServerInfo> Item);
 	TSharedRef<SWidget> MakeSelectedServerWidget();
 	FText GetSelectedServerText() const;
@@ -57,6 +57,8 @@ private:
 
 	void OnSessionNameChanged(const FText& NewName);
 	void OnSessionNameCommitted(const FText& NewText, ETextCommit::Type CommitType);
+	void OnProjectNameCommitted(const FText& NewProjectName, ETextCommit::Type CommitType);
+
 	FReply OnAccept();
 	FReply OnDecline();
 	FReply OnKeyDownHandler(const FGeometry&, const FKeyEvent&); // Registered as handler to the editable text (vs. OnKeyDown() virtual method).
@@ -67,14 +69,15 @@ private:
 	TWeakPtr<FConcertSessionItem> Item;
 	/** Servers displayed in the server combo box. */
 	TArray<TSharedPtr<FConcertServerInfo>> Servers;
-	
+
 	TSharedPtr<SComboBox<TSharedPtr<FConcertServerInfo>>> ServersComboBox;
 	TSharedPtr<SEditableTextBox> EditableSessionName;
-	
+	TSharedPtr<SEditableTextBox> ProjectNameTextBox;
+
 	FGetServersFunc GetServersFunc;
 	FAcceptFunc AcceptFunc;
 	FDeclineFunc DeclineFunc;
-	
+
 	TAttribute<FText> HighlightText;
 	TAttribute<FString> DefaultServerURL;
 	bool bInitialFocusTaken = false;
diff --git a/Engine/Plugins/Developer/Concert/ConcertUI/ConcertSharedSlate/Source/ConcertSharedSlate/Private/Session/Browser/SSessionRow.cpp b/Engine/Plugins/Developer/Concert/ConcertUI/ConcertSharedSlate/Source/ConcertSharedSlate/Private/Session/Browser/SSessionRow.cpp
--- a/Engine/Plugins/Developer/Concert/ConcertUI/ConcertSharedSlate/Source/ConcertSharedSlate/Private/Session/Browser/SSessionRow.cpp
+++ b/Engine/Plugins/Developer/Concert/ConcertUI/ConcertSharedSlate/Source/ConcertSharedSlate/Private/Session/Browser/SSessionRow.cpp
@@ -37,118 +37,194 @@ void SSessionRow::Construct(const FArguments& InArgs, TSharedPtr<FConcertSession
 	InItem->OnBeginEditSessionNameRequest.AddSP(this, &SSessionRow::OnBeginEditingSessionName);
 }
 
-TSharedRef<SWidget> SSessionRow::GenerateWidgetForColumn(const FName& ColumnName)
+
+FSlateColor SSessionRow::GetFontColor(bool bIsActiveSession, bool bIsDefault)
 {
-	TSharedPtr<FConcertSessionItem> ItemPin = Item.Pin();
+	if (bIsActiveSession)
+	{
+		return bIsDefault ? FSlateColor(FLinearColor::White) : FSlateColor(FLinearColor::White * 0.8f);
+	}
 
-	if (ColumnName == ConcertBrowserUtils::IconColName)
+	return FSlateColor::UseSubduedForeground();
+}
+
+
+FSlateFontInfo SSessionRow::GetFontInfo(bool bIsActiveSession, bool bIsDefault)
+{
+	if (bIsActiveSession)
 	{
-		return SNew(SBox)
+		return FAppStyle::Get().GetFontStyle("NormalFont");
+	}
+	return  FCoreStyle::GetDefaultFontStyle("Italic", 9);
+}
+
+TSharedRef<SWidget> SSessionRow::GenerateSessionColumn(const FSlateFontInfo& FontInfo, const FSlateColor& FontColor)
+{
+	return SNew(SBox)
+		.VAlign(VAlign_Center)
+		[
+			SAssignNew(SessionNameText, SInlineEditableTextBlock)
+			.Text_Lambda([this]() { return FText::AsCultureInvariant(Item.Pin()->SessionName); })
+			.HighlightText(HighlightText)
+			.OnTextCommitted(this, &SSessionRow::OnSessionNameCommitted)
+			.IsReadOnly(false)
+			.IsSelected(FIsSelected::CreateLambda([this]() { return IsSelected.Get(); }))
+			.OnVerifyTextChanged(this, &SSessionRow::OnValidatingSessionName)
+			.Font(FontInfo)
+			.ColorAndOpacity(FontColor)
+			];
+}
+
+TSharedRef<SWidget> SSessionRow::GenerateServerColumn(const FSlateFontInfo& FontInfo, const FSlateColor& FontColor)
+{
+	TSharedPtr<FConcertSessionItem> ItemPin = Item.Pin();
+	return SNew(SHorizontalBox)
+		+SHorizontalBox::Slot()
+		.AutoWidth()
+		.VAlign(VAlign_Center)
+		[
+			SNew(STextBlock)
+			.Text(FText::AsCultureInvariant(ItemPin->ServerName))
+			.HighlightText(HighlightText)
+			.Font(FontInfo)
+			.ColorAndOpacity(FontColor)
+		]
+		+SHorizontalBox::Slot()
+		[
+			SNew(SSpacer)
+		]
+		+SHorizontalBox::Slot()
+		.AutoWidth()
+		.VAlign(VAlign_Center)
+		.HAlign(HAlign_Right)
+		[
+			ConcertBrowserUtils::MakeServerVersionIgnoredWidget(ItemPin->ServerFlags)
+		];
+}
+
+TSharedRef<SWidget> SSessionRow::GenerateServerDefaultColumn(const FSlateFontInfo& FontInfo, const FSlateColor& FontColor)
+{
+	TSharedPtr<FConcertSessionItem> ItemPin = Item.Pin();
+	return SNew(SHorizontalBox)
+		+SHorizontalBox::Slot()
+		.AutoWidth()
+		.VAlign(VAlign_Center)
+		[
+			SNew(STextBlock)
+			.Text(FText::Format(INVTEXT("{0} * "), FText::AsCultureInvariant(ItemPin->ServerName)))
+			.HighlightText(HighlightText)
+			.Font(FontInfo)
+			.ColorAndOpacity(FontColor)
+		]
+		+SHorizontalBox::Slot()
+		.AutoWidth()
+		.VAlign(VAlign_Center)
+		[
+			SNew(STextBlock)
+			.Text(LOCTEXT("DefaultServerSession", "(Default Session/Server)"))
+			.HighlightText(HighlightText)
+			.Font(FCoreStyle::GetDefaultFontStyle("Regular", 9))
+			.ColorAndOpacity(FontColor)
+		]
+		+SHorizontalBox::Slot()
+		[
+			SNew(SSpacer)
+		]
+		+SHorizontalBox::Slot()
+		.AutoWidth()
+		.VAlign(VAlign_Center)
+		.HAlign(HAlign_Right)
+		[
+			ConcertBrowserUtils::MakeServerVersionIgnoredWidget(ItemPin->ServerFlags)
+		];
+}
+
+TSharedRef<SWidget> SSessionRow::GenerateProjectColumn(const FSlateFontInfo &FontInfo, const FSlateColor &FontColor)
+{
+	//TSharedPtr<FConcertSessionItem> ItemPin = Item.Pin();
+	return SNew(SBox)
+		.VAlign(VAlign_Center)
+		[
+			SAssignNew(SessionNameText, SInlineEditableTextBlock)
+			.Text_Lambda([this]() { return FText::AsCultureInvariant(Item.Pin()->ProjectName); })
+			.HighlightText(HighlightText)
+			.IsReadOnly(true)
+			.Font(FontInfo)
+			.ColorAndOpacity(FontColor)
+			];
+}
+
+TSharedRef<SWidget> SSessionRow::GenerateVersionColumn(const FSlateFontInfo& FontInfo, const FSlateColor& FontColor)
+{
+	return SNew(SBox)
+		.VAlign(VAlign_Center)
+		[
+			SAssignNew(SessionNameText, SInlineEditableTextBlock)
+			.Text_Lambda([this]() { return FText::AsCultureInvariant(Item.Pin()->ProjectVersion); })
+			.HighlightText(HighlightText)
+			.IsReadOnly(true)
+			.Font(FontInfo)
+			.ColorAndOpacity(FontColor)
+			];
+}
+
+
+TSharedRef<SWidget> SSessionRow::GenerateIconColumn()
+{
+	TSharedPtr<FConcertSessionItem> ItemPin = Item.Pin();
+	const FText Tooltip = ItemPin->Type == FConcertSessionItem::EType::ActiveSession ?
+		LOCTEXT("ActiveIconTooltip", "Active session") : LOCTEXT("ArchivedIconTooltip", "Archived Session");
+	const FSlateColor InStyle = ItemPin->Type == FConcertSessionItem::EType::ActiveSession ?
+		FAppStyle::Get().GetWidgetStyle<FButtonStyle>("FlatButton.Success").Normal.TintColor : FSlateColor::UseSubduedForeground();
+	const FText InGlyph = ItemPin->Type == FConcertSessionItem::EType::ActiveSession ? FEditorFontGlyphs::Circle : FEditorFontGlyphs::Archive;
+	return SNew(SBox)
 		.VAlign(VAlign_Center)
 		.HAlign(HAlign_Center)
 		.Padding(2)
-		.ToolTipText(ItemPin->Type == FConcertSessionItem::EType::ActiveSession ? LOCTEXT("ActiveIconTooltip", "Active session") : LOCTEXT("ArchivedIconTooltip", "Archived Session"))
+		.ToolTipText(Tooltip)
 		[
 			SNew(STextBlock)
 			.Font(FAppStyle::Get().GetFontStyle(ConcertBrowserUtils::IconColumnFontName))
-			.Text(ItemPin->Type == FConcertSessionItem::EType::ActiveSession ? FEditorFontGlyphs::Circle : FEditorFontGlyphs::Archive)
-			.ColorAndOpacity(ItemPin->Type == FConcertSessionItem::EType::ActiveSession ? FAppStyle::Get().GetWidgetStyle<FButtonStyle>("FlatButton.Success").Normal.TintColor : FSlateColor::UseSubduedForeground())
+			.Text(InGlyph)
+			.ColorAndOpacity(InStyle)
 		];
-	}
-	
-	const bool bIsDefaultConfig = IsDefaultSession(ItemPin);
-	FSlateFontInfo FontInfo;
-	FSlateColor FontColor;
-	if (ItemPin->Type == FConcertSessionItem::EType::ActiveSession)
-	{
-		FontColor = bIsDefaultConfig ? FSlateColor(FLinearColor::White) : FSlateColor(FLinearColor::White * 0.8f);
-		FontInfo = FAppStyle::Get().GetFontStyle("NormalFont");
-	}
-	else
+}
+
+TSharedRef<SWidget> SSessionRow::GenerateWidgetForColumn(const FName& ColumnName)
+{
+	if (ColumnName == ConcertBrowserUtils::IconColName)
 	{
-		FontColor = FSlateColor::UseSubduedForeground();
-		FontInfo = FCoreStyle::GetDefaultFontStyle("Italic", 9);
+		return GenerateIconColumn();
 	}
 
+	TSharedPtr<FConcertSessionItem> ItemPin = Item.Pin();
+	const bool bIsActiveSession = ItemPin->Type == FConcertSessionItem::EType::ActiveSession;
+	const bool bIsDefaultConfig = IsDefaultSession(ItemPin);
+	FSlateFontInfo FontInfo = GetFontInfo(bIsActiveSession, bIsDefaultConfig);
+	FSlateColor FontColor = GetFontColor(bIsActiveSession, bIsDefaultConfig);
+
 	if (ColumnName == ConcertBrowserUtils::SessionColName)
 	{
-		return SNew(SBox)
-			.VAlign(VAlign_Center)
-			[
-				SAssignNew(SessionNameText, SInlineEditableTextBlock)
-				.Text_Lambda([this]() { return FText::AsCultureInvariant(Item.Pin()->SessionName); })
-				.HighlightText(HighlightText)
-				.OnTextCommitted(this, &SSessionRow::OnSessionNameCommitted)
-				.IsReadOnly(false)
-				.IsSelected(FIsSelected::CreateLambda([this]() { return IsSelected.Get(); }))
-				.OnVerifyTextChanged(this, &SSessionRow::OnValidatingSessionName)
-				.Font(FontInfo)
-				.ColorAndOpacity(FontColor)
-			];
+		return GenerateSessionColumn(FontInfo, FontColor);
 	}
 
-	check(ColumnName == ConcertBrowserUtils::ServerColName);
-
-	if (bIsDefaultConfig)
+	if (ColumnName == ConcertBrowserUtils::ServerColName)
 	{
-		return SNew(SHorizontalBox)
-			+SHorizontalBox::Slot()
-			.AutoWidth()
-			.VAlign(VAlign_Center)
-			[
-				SNew(STextBlock)
-				.Text(FText::Format(INVTEXT("{0} * "), FText::AsCultureInvariant(ItemPin->ServerName)))
-				.HighlightText(HighlightText)
-				.Font(FontInfo)
-				.ColorAndOpacity(FontColor)
-			]
-			+SHorizontalBox::Slot()
-			.AutoWidth()
-			.VAlign(VAlign_Center)
-			[
-				SNew(STextBlock)
-				.Text(LOCTEXT("DefaultServerSession", "(Default Session/Server)"))
-				.HighlightText(HighlightText)
-				.Font(FCoreStyle::GetDefaultFontStyle("Regular", 9))
-				.ColorAndOpacity(FontColor)
-			]
-			+SHorizontalBox::Slot()
-			[
-				SNew(SSpacer)
-			]
-			+SHorizontalBox::Slot()
-			.AutoWidth()
-			.VAlign(VAlign_Center)
-			.HAlign(HAlign_Right)
-			[
-				ConcertBrowserUtils::MakeServerVersionIgnoredWidget(ItemPin->ServerFlags)
-			];
+		if (bIsDefaultConfig)
+		{
+			return GenerateServerDefaultColumn(FontInfo, FontColor);
+		}
+		return GenerateServerColumn(FontInfo, FontColor);
 	}
-	else
+
+	if (ColumnName == ConcertBrowserUtils::ProjectColName)
 	{
-		return SNew(SHorizontalBox)
-			+SHorizontalBox::Slot()
-			.AutoWidth()
-			.VAlign(VAlign_Center)
-			[
-				SNew(STextBlock)
-				.Text(FText::AsCultureInvariant(ItemPin->ServerName))
-				.HighlightText(HighlightText)
-				.Font(FontInfo)
-				.ColorAndOpacity(FontColor)
-			]
-			+SHorizontalBox::Slot()
-			[
-				SNew(SSpacer)
-			]
-			+SHorizontalBox::Slot()
-			.AutoWidth()
-			.VAlign(VAlign_Center)
-			.HAlign(HAlign_Right)
-			[
-				ConcertBrowserUtils::MakeServerVersionIgnoredWidget(ItemPin->ServerFlags)
-			];
+		return GenerateProjectColumn(FontInfo, FontColor);
 	}
+
+	check(ColumnName == ConcertBrowserUtils::VersionColName);
+
+	return GenerateVersionColumn(FontInfo, FontColor);
 }
 
 FReply SSessionRow::OnMouseButtonDoubleClick(const FGeometry& InMyGeometry, const FPointerEvent& InMouseEvent)
@@ -186,4 +262,4 @@ void SSessionRow::OnSessionNameCommitted(const FText& NewSessionName, ETextCommi
 	}
 }
 
-#undef LOCTEXT_NAMESPACE
\ No newline at end of file
+#undef LOCTEXT_NAMESPACE
diff --git a/Engine/Plugins/Developer/Concert/ConcertUI/ConcertSharedSlate/Source/ConcertSharedSlate/Private/Session/Browser/SSessionRow.h b/Engine/Plugins/Developer/Concert/ConcertUI/ConcertSharedSlate/Source/ConcertSharedSlate/Private/Session/Browser/SSessionRow.h
--- a/Engine/Plugins/Developer/Concert/ConcertUI/ConcertSharedSlate/Source/ConcertSharedSlate/Private/Session/Browser/SSessionRow.h
+++ b/Engine/Plugins/Developer/Concert/ConcertUI/ConcertSharedSlate/Source/ConcertSharedSlate/Private/Session/Browser/SSessionRow.h
@@ -42,9 +42,18 @@ public:
 	virtual FReply OnMouseButtonDoubleClick( const FGeometry& InMyGeometry, const FPointerEvent& InMouseEvent) override;
 
 	void OnSessionNameCommitted(const FText& NewSessionName, ETextCommit::Type CommitType);
-
-private:
 	
+private:
+	FSlateFontInfo GetFontInfo(bool bIsActiveSession, bool bIsDefault);
+	FSlateColor    GetFontColor(bool bIsActiveSession, bool bIsDefault);
+
+	TSharedRef<SWidget> GenerateIconColumn();
+	TSharedRef<SWidget> GenerateSessionColumn(const FSlateFontInfo& FontInfo, const FSlateColor& FontColor);
+	TSharedRef<SWidget> GenerateServerColumn(const FSlateFontInfo& FontInfo, const FSlateColor& FontColor);
+	TSharedRef<SWidget> GenerateServerDefaultColumn(const FSlateFontInfo& FontInfo, const FSlateColor& FontColor);
+	TSharedRef<SWidget> GenerateProjectColumn(const FSlateFontInfo& FontInfo, const FSlateColor& FontColor);
+	TSharedRef<SWidget> GenerateVersionColumn(const FSlateFontInfo& FontInfo, const FSlateColor& FontColor);
+
 	void OnBeginEditingSessionName() { SessionNameText->EnterEditingMode(); }
 	bool OnValidatingSessionName(const FText& NewSessionName, FText& OutError);
 
diff --git a/Engine/Plugins/Developer/Concert/ConcertUI/ConcertSharedSlate/Source/ConcertSharedSlate/Public/Session/Browser/ConcertBrowserUtils.h b/Engine/Plugins/Developer/Concert/ConcertUI/ConcertSharedSlate/Source/ConcertSharedSlate/Public/Session/Browser/ConcertBrowserUtils.h
--- a/Engine/Plugins/Developer/Concert/ConcertUI/ConcertSharedSlate/Source/ConcertSharedSlate/Public/Session/Browser/ConcertBrowserUtils.h
+++ b/Engine/Plugins/Developer/Concert/ConcertUI/ConcertSharedSlate/Source/ConcertSharedSlate/Public/Session/Browser/ConcertBrowserUtils.h
@@ -26,6 +26,8 @@ namespace ConcertBrowserUtils
 	static const FName IconColName(TEXT("Icon"));
 	static const FName SessionColName(TEXT("Session"));
 	static const FName ServerColName(TEXT("Server"));
+	static const FName ProjectColName(TEXT("Project"));
+	static const FName VersionColName(TEXT("Version"));
 
 	// Name of the filter box in the View option.
 	static const FName ActiveSessionsCheckBoxMenuName(TEXT("ActiveSessions"));
diff --git a/Engine/Plugins/Developer/Concert/ConcertUI/ConcertSharedSlate/Source/ConcertSharedSlate/Public/Session/Browser/ConcertSessionItem.h b/Engine/Plugins/Developer/Concert/ConcertUI/ConcertSharedSlate/Source/ConcertSharedSlate/Public/Session/Browser/ConcertSessionItem.h
--- a/Engine/Plugins/Developer/Concert/ConcertUI/ConcertSharedSlate/Source/ConcertSharedSlate/Public/Session/Browser/ConcertSessionItem.h
+++ b/Engine/Plugins/Developer/Concert/ConcertUI/ConcertSharedSlate/Source/ConcertSharedSlate/Public/Session/Browser/ConcertSessionItem.h
@@ -24,19 +24,16 @@ public:
 		ArchivedSession, // Read-only item representing an archived session.
 	};
 
-	FConcertSessionItem(EType Type, const FString& InSessionName, const FGuid& InSessionId, const FString& InServerName, const FGuid& InServerEndpoint, EConcertServerFlags InServerFlags)
-		: Type(Type)
-		, ServerAdminEndpointId(InServerEndpoint)
-		, SessionId(InSessionId)
-		, SessionName(InSessionName)
-		, ServerName(InServerName)
-		, ServerFlags(InServerFlags)
+	bool operator==(const FConcertSessionItem& Other) const
 	{
+		return Type == Other.Type && ServerAdminEndpointId == Other.ServerAdminEndpointId && SessionId == Other.SessionId;
 	}
 
-	bool operator==(const FConcertSessionItem& Other) const
+	FConcertSessionItem MakeCopyAsType(EType NewType)
 	{
-		return Type == Other.Type && ServerAdminEndpointId == Other.ServerAdminEndpointId && SessionId == Other.SessionId;
+		FConcertSessionItem NewItem = *this;
+		Type = NewType;
+		return NewItem;
 	}
 
 	EType Type = EType::None;
@@ -44,6 +41,9 @@ public:
 	FGuid SessionId;
 	FString SessionName;
 	FString ServerName;
+	FString ProjectName;
+	FString ProjectVersion;
+	EConcertServerFlags ServerFlags = EConcertServerFlags::None;
+
 	FOnBeginEditConcertSessionNameRequest OnBeginEditSessionNameRequest; // Emitted when user press 'F2' or select 'Rename' from context menu.
-	EConcertServerFlags ServerFlags;
-};
\ No newline at end of file
+};
diff --git a/Engine/Plugins/Developer/Concert/ConcertUI/ConcertSharedSlate/Source/ConcertSharedSlate/Public/Session/Browser/IConcertSessionBrowserController.h b/Engine/Plugins/Developer/Concert/ConcertUI/ConcertSharedSlate/Source/ConcertSharedSlate/Public/Session/Browser/IConcertSessionBrowserController.h
--- a/Engine/Plugins/Developer/Concert/ConcertUI/ConcertSharedSlate/Source/ConcertSharedSlate/Public/Session/Browser/IConcertSessionBrowserController.h
+++ b/Engine/Plugins/Developer/Concert/ConcertUI/ConcertSharedSlate/Source/ConcertSharedSlate/Public/Session/Browser/IConcertSessionBrowserController.h
@@ -39,7 +39,6 @@ public:
 	};
 
 public:
-
 	/** Returns the latest list of server known to this controller. */
 	virtual TArray<FConcertServerInfo> GetServers() const = 0;
 
@@ -55,7 +54,7 @@ public:
 	/** Returns the archived sessions info corresponding to the specified parameters. Used to display the sessions details. */
 	virtual TOptional<FConcertSessionInfo> GetArchivedSessionInfo(const FGuid& AdminEndpoint, const FGuid& SessionId) const = 0;
 
-	virtual void CreateSession(const FGuid& ServerAdminEndpointId, const FString& SessionName) = 0;
+	virtual void CreateSession(const FGuid& ServerAdminEndpointId, const FString& SessionName, const FString& ProjectName) = 0;
 	virtual void ArchiveSession(const FGuid& ServerAdminEndpointId, const FGuid& SessionId, const FString& ArchiveName, const FConcertSessionFilter& SessionFilter) = 0;
 	virtual void RestoreSession(const FGuid& ServerAdminEndpointId, const FGuid& SessionId, const FString& RestoredName, const FConcertSessionFilter& SessionFilter) = 0;
 	virtual void RenameActiveSession(const FGuid& ServerAdminEndpointId, const FGuid& SessionId, const FString& NewName) = 0;
@@ -68,4 +67,4 @@ public:
 	virtual bool CanDeleteArchivedSession(const FGuid& ServerAdminEndpointId, const FGuid& SessionId) const = 0;
 	
 	virtual ~IConcertSessionBrowserController() = default;
-};
\ No newline at end of file
+};
