UI for restoring and archiving sessions is broken

diff --git a/Engine/Plugins/Developer/Concert/ConcertUI/ConcertSharedSlate/Source/ConcertSharedSlate/Private/Session/Browser/SConcertSessionBrowser.cpp b/Engine/Plugins/Developer/Concert/ConcertUI/ConcertSharedSlate/Source/ConcertSharedSlate/Private/Session/Browser/SConcertSessionBrowser.cpp
--- a/Engine/Plugins/Developer/Concert/ConcertUI/ConcertSharedSlate/Source/ConcertSharedSlate/Private/Session/Browser/SConcertSessionBrowser.cpp
+++ b/Engine/Plugins/Developer/Concert/ConcertUI/ConcertSharedSlate/Source/ConcertSharedSlate/Private/Session/Browser/SConcertSessionBrowser.cpp
@@ -601,11 +601,13 @@ void SConcertSessionBrowser::InsertNewSessionEditableRowInternal()
 	InsertEditableSessionRow(MakeShared<FConcertSessionItem>(MoveTemp(Item)), nullptr);
 }
 
+PRAGMA_DISABLE_OPTIMIZATION
 void SConcertSessionBrowser::InsertRestoreSessionAsEditableRowInternal(const TSharedPtr<FConcertSessionItem>& ArchivedItem)
 {
 	// Insert the 'restore session as ' editable row just below the 'archived' item to restore.
 	InsertEditableSessionRow(MakeShared<FConcertSessionItem>(ArchivedItem->MakeCopyAsType(FConcertSessionItem::EType::RestoreSession)), ArchivedItem);
 }
+	PRAGMA_ENABLE_OPTIMIZATION
 
 void SConcertSessionBrowser::InsertArchiveSessionAsEditableRow(const TSharedPtr<FConcertSessionItem>& LiveItem)
 {
diff --git a/Engine/Plugins/Developer/Concert/ConcertUI/ConcertSharedSlate/Source/ConcertSharedSlate/Private/Session/Browser/SSaveRestoreSessionRow.cpp b/Engine/Plugins/Developer/Concert/ConcertUI/ConcertSharedSlate/Source/ConcertSharedSlate/Private/Session/Browser/SSaveRestoreSessionRow.cpp
--- a/Engine/Plugins/Developer/Concert/ConcertUI/ConcertSharedSlate/Source/ConcertSharedSlate/Private/Session/Browser/SSaveRestoreSessionRow.cpp
+++ b/Engine/Plugins/Developer/Concert/ConcertUI/ConcertSharedSlate/Source/ConcertSharedSlate/Private/Session/Browser/SSaveRestoreSessionRow.cpp
@@ -100,10 +100,8 @@ TSharedRef<SWidget> SSaveRestoreSessionRow::GenerateWidgetForColumn(const FName&
 				.SelectAllTextWhenFocused(true)
 			];
 	}
-	else
+	else if (ColumnName == ConcertBrowserUtils::ServerColName)
 	{
-		check(ColumnName == ConcertBrowserUtils::ServerColName);
-
 		// 'Server' text block.
 		return SNew(SHorizontalBox)
 		+SHorizontalBox::Slot()
@@ -148,6 +146,11 @@ TSharedRef<SWidget> SSaveRestoreSessionRow::GenerateWidgetForColumn(const FName&
 			]
 		];
 	}
+	else
+	{
+		check(ColumnName == ConcertBrowserUtils::ProjectColName || ColumnName == ConcertBrowserUtils::VersionColName);
+		return SNullWidget::NullWidget;
+	}
 }
 
 void SSaveRestoreSessionRow::Tick(const FGeometry& AllottedGeometry, const double InCurrentTime, const float InDeltaTime)
diff --git a/Engine/Plugins/Developer/Concert/ConcertUI/ConcertSharedSlate/Source/ConcertSharedSlate/Public/Session/Browser/ConcertSessionItem.h b/Engine/Plugins/Developer/Concert/ConcertUI/ConcertSharedSlate/Source/ConcertSharedSlate/Public/Session/Browser/ConcertSessionItem.h
--- a/Engine/Plugins/Developer/Concert/ConcertUI/ConcertSharedSlate/Source/ConcertSharedSlate/Public/Session/Browser/ConcertSessionItem.h
+++ b/Engine/Plugins/Developer/Concert/ConcertUI/ConcertSharedSlate/Source/ConcertSharedSlate/Public/Session/Browser/ConcertSessionItem.h
@@ -29,10 +29,10 @@ public:
 		return Type == Other.Type && ServerAdminEndpointId == Other.ServerAdminEndpointId && SessionId == Other.SessionId;
 	}
 
-	FConcertSessionItem MakeCopyAsType(EType NewType)
+	FConcertSessionItem MakeCopyAsType(EType NewType) const
 	{
 		FConcertSessionItem NewItem = *this;
-		Type = NewType;
+		NewItem.Type = NewType;
 		return NewItem;
 	}
 
