Display delete symbol on activity only when hovered

diff --git a/Engine/Plugins/Developer/Concert/ConcertUI/ConcertSharedSlate/Source/ConcertSharedSlate/Private/SConcertSessionRecovery.cpp b/Engine/Plugins/Developer/Concert/ConcertUI/ConcertSharedSlate/Source/ConcertSharedSlate/Private/SConcertSessionRecovery.cpp
--- a/Engine/Plugins/Developer/Concert/ConcertUI/ConcertSharedSlate/Source/ConcertSharedSlate/Private/SConcertSessionRecovery.cpp
+++ b/Engine/Plugins/Developer/Concert/ConcertUI/ConcertSharedSlate/Source/ConcertSharedSlate/Private/SConcertSessionRecovery.cpp
@@ -35,7 +35,7 @@ void SConcertSessionRecovery::Construct(const FArguments& InArgs)
 	SAssignNew(ActivityView, SConcertSessionActivities)
 		.OnFetchActivities(InArgs._OnFetchActivities)
 		.OnMapActivityToClient(InArgs._OnMapActivityToClient)
-		.OnMakeColumnOverlayWidget_Lambda([this](TWeakPtr<FConcertSessionActivity> Activity, const FName& ColumnId) { return MakeRecoverThroughWidget(Activity, ColumnId); })
+		.OnMakeColumnOverlayWidget_Lambda([this](TWeakPtr<SMultiColumnTableRow<TSharedPtr<FConcertSessionActivity>>>, TWeakPtr<FConcertSessionActivity> Activity, const FName& ColumnId) { return MakeRecoverThroughWidget(Activity, ColumnId); })
 		.Columns(MakeColumns(InArgs))
 		.HighlightText(this, &SConcertSessionRecovery::HighlightSearchText)
 		.TimeFormat(ActivityViewOptions.Get(), &FConcertSessionActivitiesOptions::GetTimeFormat)
diff --git a/Engine/Plugins/Developer/Concert/ConcertUI/ConcertSharedSlate/Source/ConcertSharedSlate/Private/Session/Activity/SConcertSessionActivities.cpp b/Engine/Plugins/Developer/Concert/ConcertUI/ConcertSharedSlate/Source/ConcertSharedSlate/Private/Session/Activity/SConcertSessionActivities.cpp
--- a/Engine/Plugins/Developer/Concert/ConcertUI/ConcertSharedSlate/Source/ConcertSharedSlate/Private/Session/Activity/SConcertSessionActivities.cpp
+++ b/Engine/Plugins/Developer/Concert/ConcertUI/ConcertSharedSlate/Source/ConcertSharedSlate/Private/Session/Activity/SConcertSessionActivities.cpp
@@ -154,7 +154,7 @@ TSharedRef<SWidget> SConcertSessionActivityRow::GenerateWidgetForColumn(const FN
 	
 	if (OnMakeColumnOverlayWidget.IsBound())
 	{
-		if (TSharedPtr<SWidget> OverlayedWidget = OnMakeColumnOverlayWidget.Execute(ActivityPin, ColumnId))
+		if (TSharedPtr<SWidget> OverlayedWidget = OnMakeColumnOverlayWidget.Execute(SharedThis(this), ActivityPin, ColumnId))
 		{
 			Overlay->AddSlot()
 			[
diff --git a/Engine/Plugins/Developer/Concert/ConcertUI/ConcertSharedSlate/Source/ConcertSharedSlate/Private/Session/History/SEditableSessionHistory.cpp b/Engine/Plugins/Developer/Concert/ConcertUI/ConcertSharedSlate/Source/ConcertSharedSlate/Private/Session/History/SEditableSessionHistory.cpp
--- a/Engine/Plugins/Developer/Concert/ConcertUI/ConcertSharedSlate/Source/ConcertSharedSlate/Private/Session/History/SEditableSessionHistory.cpp
+++ b/Engine/Plugins/Developer/Concert/ConcertUI/ConcertSharedSlate/Source/ConcertSharedSlate/Private/Session/History/SEditableSessionHistory.cpp
@@ -26,13 +26,16 @@ void SEditableSessionHistory::Construct(const FArguments& InArgs)
 	];
 }
 
-TSharedPtr<SWidget> SEditableSessionHistory::MakeSummaryColumnDeleteButton(TWeakPtr<FConcertSessionActivity> RowActivity, const FName& ColumnId) const
+TSharedPtr<SWidget> SEditableSessionHistory::MakeSummaryColumnDeleteButton(
+	TWeakPtr<SMultiColumnTableRow<TSharedPtr<FConcertSessionActivity>>> OwningRow,
+	TWeakPtr<FConcertSessionActivity> RowActivity,
+	const FName& ColumnId) const
 {
 	if (!SessionHistory->IsLastColumn(ColumnId))
 	{
 		return nullptr;
 	}
-	
+
 	return SNew(SBox)
 		.Padding(FMargin(1, 1))
 		.HAlign(HAlign_Right)
@@ -40,6 +43,7 @@ TSharedPtr<SWidget> SEditableSessionHistory::MakeSummaryColumnDeleteButton(TWeak
 		[
 			SNew(SNegativeActionButton)
 			.OnClicked(this, &SEditableSessionHistory::OnClickDeleteActivityButton, RowActivity)
+			.Visibility_Lambda([OwningRow](){ return OwningRow.IsValid() && OwningRow.Pin()->IsHovered() ? EVisibility::Visible : EVisibility::Collapsed; })
 			.ToolTipText(TAttribute<FText>::CreateLambda([this, RowActivity](){ return GetDeleteActivityToolTip(RowActivity);}))
 			.IsEnabled(TAttribute<bool>::CreateLambda([this, RowActivity](){ return CanDeleteActivity(RowActivity);}))
 			.Icon(FEditorStyle::GetBrush("Icons.Delete"))
diff --git a/Engine/Plugins/Developer/Concert/ConcertUI/ConcertSharedSlate/Source/ConcertSharedSlate/Public/Session/Activity/SConcertSessionActivities.h b/Engine/Plugins/Developer/Concert/ConcertUI/ConcertSharedSlate/Source/ConcertSharedSlate/Public/Session/Activity/SConcertSessionActivities.h
--- a/Engine/Plugins/Developer/Concert/ConcertUI/ConcertSharedSlate/Source/ConcertSharedSlate/Public/Session/Activity/SConcertSessionActivities.h
+++ b/Engine/Plugins/Developer/Concert/ConcertUI/ConcertSharedSlate/Source/ConcertSharedSlate/Public/Session/Activity/SConcertSessionActivities.h
@@ -47,19 +47,34 @@ public:
 	};
 
 	/** Used to pull activities from a session. Used to fetch and display the activities of an archived session. */
-	DECLARE_DELEGATE_RetVal_ThreeParams(bool, FFetchActivitiesFunc, TArray<TSharedPtr<FConcertSessionActivity>>& /*InOutActivities*/, int32& /*OutFetchedCount*/, FText& /*ErrorMsg*/);
+	DECLARE_DELEGATE_RetVal_ThreeParams(bool, FFetchActivitiesFunc,
+		TArray<TSharedPtr<FConcertSessionActivity>>& /*InOutActivities*/,
+		int32& /*OutFetchedCount*/,
+		FText& /*ErrorMsg*/
+		)
 	
 	/** Used to map an activity to its client. */
-	DECLARE_DELEGATE_RetVal_OneParam(TOptional<FConcertClientInfo>, FGetActivityClientInfoFunc, FGuid /*ClientId*/);
+	DECLARE_DELEGATE_RetVal_OneParam(TOptional<FConcertClientInfo>, FGetActivityClientInfoFunc,
+		FGuid /*ClientId*/
+		)
 
 	/** Returns the transaction event corresponding the specified activity.*/
-	DECLARE_DELEGATE_RetVal_OneParam(TFuture<TOptional<FConcertSyncTransactionEvent>>, FGetTransactionEvent, const FConcertSessionActivity& /*Activity*/)
+	DECLARE_DELEGATE_RetVal_OneParam(TFuture<TOptional<FConcertSyncTransactionEvent>>, FGetTransactionEvent,
+		const FConcertSessionActivity& /*Activity*/
+		)
 	
 	/** Returns the package event corresponding to the package activity. */
-	DECLARE_DELEGATE_RetVal_TwoParams(bool, FGetPackageEvent, const FConcertSessionActivity& /*Activity*/, FConcertSyncPackageEventMetaData& /*OutEvent*/);
+	DECLARE_DELEGATE_RetVal_TwoParams(bool, FGetPackageEvent,
+		const FConcertSessionActivity& /*Activity*/,
+		FConcertSyncPackageEventMetaData& /*OutEvent*/
+		)
 
 	/** Used to overlay a widget over a column widget to add custom functionalities to a row. */
-	DECLARE_DELEGATE_RetVal_TwoParams(TSharedPtr<SWidget>, FMakeColumnOverlayWidgetFunc, TWeakPtr<FConcertSessionActivity> /*ThisRowActivity*/, const FName& /*ColumnId*/);
+	DECLARE_DELEGATE_RetVal_ThreeParams(TSharedPtr<SWidget>, FMakeColumnOverlayWidgetFunc,
+		TWeakPtr<SMultiColumnTableRow<TSharedPtr<FConcertSessionActivity>>> /*OwningRow*/,
+		TWeakPtr<FConcertSessionActivity> /*ThisRowActivity*/,
+		const FName& /*ColumnId*/
+		)
 
 public:
 	SLATE_BEGIN_ARGS(SConcertSessionActivities)
diff --git a/Engine/Plugins/Developer/Concert/ConcertUI/ConcertSharedSlate/Source/ConcertSharedSlate/Public/Session/History/SEditableSessionHistory.h b/Engine/Plugins/Developer/Concert/ConcertUI/ConcertSharedSlate/Source/ConcertSharedSlate/Public/Session/History/SEditableSessionHistory.h
--- a/Engine/Plugins/Developer/Concert/ConcertUI/ConcertSharedSlate/Source/ConcertSharedSlate/Public/Session/History/SEditableSessionHistory.h
+++ b/Engine/Plugins/Developer/Concert/ConcertUI/ConcertSharedSlate/Source/ConcertSharedSlate/Public/Session/History/SEditableSessionHistory.h
@@ -37,7 +37,7 @@ private:
 	FRequestDeleteActivity DeleteActivityFunc;
 
 	/** Creates a delete button that is overlayed over the summary column*/
-	TSharedPtr<SWidget> MakeSummaryColumnDeleteButton(TWeakPtr<FConcertSessionActivity> RowActivity, const FName& ColumnId) const;
+	TSharedPtr<SWidget> MakeSummaryColumnDeleteButton(TWeakPtr<SMultiColumnTableRow<TSharedPtr<FConcertSessionActivity>>> OwningRow, TWeakPtr<FConcertSessionActivity> RowActivity, const FName& ColumnId) const;
 	FReply OnClickDeleteActivityButton(TWeakPtr<FConcertSessionActivity> RowActivity) const;
 	FText GetDeleteActivityToolTip(TWeakPtr<FConcertSessionActivity> RowActivity) const;
 	bool CanDeleteActivity(TWeakPtr<FConcertSessionActivity> RowActivity) const;
