Hide and restore optional columns.

diff --git a/Engine/Plugins/Developer/Concert/ConcertApp/MultiUserServer/Source/MultiUserServer/Private/MultiUserServerUserSettings.cpp b/Engine/Plugins/Developer/Concert/ConcertApp/MultiUserServer/Source/MultiUserServer/Private/MultiUserServerUserSettings.cpp
new file mode 100644
--- /dev/null
+++ b/Engine/Plugins/Developer/Concert/ConcertApp/MultiUserServer/Source/MultiUserServer/Private/MultiUserServerUserSettings.cpp
@@ -0,0 +1,43 @@
+﻿// Copyright Epic Games, Inc. All Rights Reserved.
+
+#include "MultiUserServerUserSettings.h"
+
+#include "Misc/CoreDelegates.h"
+
+namespace UE::MultiUserServer
+{
+	static bool bIsShutdown = false;
+}
+
+UMultiUserServerUserSettings::UMultiUserServerUserSettings()
+{
+	OnSessionBrowserColumnVisibilityChanged().AddLambda([this](const FColumnVisibilitySnapshot&)
+	{
+		SaveConfig();
+	});
+	OnArchivedActivityBrowserColumnVisibility().AddLambda([this](const FColumnVisibilitySnapshot&)
+	{
+		SaveConfig();
+	});
+	OnLiveActivityBrowserColumnVisibility().AddLambda([this](const FColumnVisibilitySnapshot&)
+	{
+		SaveConfig();
+	});
+	OnLiveSessionContentColumnVisibility().AddLambda([this](const FColumnVisibilitySnapshot&)
+	{
+		SaveConfig();
+	});
+
+	FCoreDelegates::OnPreExit.AddLambda([]()
+	{
+		UE::MultiUserServer::bIsShutdown = true;
+	});
+}
+
+UMultiUserServerUserSettings* UMultiUserServerUserSettings::GetUserSettings()
+{
+	// If we're shutting down GetMutableDefault will return garbage - this function may be called by destructor when this module is unloaded
+	return UE::MultiUserServer::bIsShutdown
+		? nullptr
+		: GetMutableDefault<UMultiUserServerUserSettings>();
+}
diff --git a/Engine/Plugins/Developer/Concert/ConcertApp/MultiUserServer/Source/MultiUserServer/Private/MultiUserServerUserSettings.h b/Engine/Plugins/Developer/Concert/ConcertApp/MultiUserServer/Source/MultiUserServer/Private/MultiUserServerUserSettings.h
new file mode 100644
--- /dev/null
+++ b/Engine/Plugins/Developer/Concert/ConcertApp/MultiUserServer/Source/MultiUserServer/Private/MultiUserServerUserSettings.h
@@ -0,0 +1,58 @@
+﻿// Copyright Epic Games, Inc. All Rights Reserved.
+
+#pragma once
+
+#include "CoreMinimal.h"
+#include "UObject/Object.h"
+#include "ConcertHeaderRowUtils.h"
+#include "MultiUserServerUserSettings.generated.h"
+
+/**
+ * 
+ */
+UCLASS(Config = MultiUserServerUserSettings, DefaultConfig)
+class MULTIUSERSERVER_API UMultiUserServerUserSettings : public UObject
+{
+	GENERATED_BODY()
+public:
+
+	DECLARE_MULTICAST_DELEGATE_OneParam(FOnColumnVisibilitySnapshotChanged, const FColumnVisibilitySnapshot& /*NewValue*/);
+
+	UMultiUserServerUserSettings();
+	
+	static UMultiUserServerUserSettings* GetUserSettings();
+
+	const FColumnVisibilitySnapshot& GetSessionBrowserColumnVisibility() const { return SessionBrowserColumnVisibility; }
+	void SetSessionBrowserColumnVisibility(FColumnVisibilitySnapshot NewValue) { SessionBrowserColumnVisibility = MoveTemp(NewValue); OnSessionBrowserColumnVisibilityChangedEvent.Broadcast(SessionBrowserColumnVisibility); }
+	FOnColumnVisibilitySnapshotChanged& OnSessionBrowserColumnVisibilityChanged() { return OnSessionBrowserColumnVisibilityChangedEvent; }
+
+	const FColumnVisibilitySnapshot& GetArchivedActivityBrowserColumnVisibility() const { return ArchivedActivityBrowserColumnVisibility; }
+	void SetArchivedActivityBrowserColumnVisibility(FColumnVisibilitySnapshot NewValue) { ArchivedActivityBrowserColumnVisibility = MoveTemp(NewValue); OnArchivedActivityBrowserColumnVisibilityEvent.Broadcast(ArchivedActivityBrowserColumnVisibility); }
+	FOnColumnVisibilitySnapshotChanged& OnArchivedActivityBrowserColumnVisibility() { return OnArchivedActivityBrowserColumnVisibilityEvent; }
+
+	const FColumnVisibilitySnapshot& GetLiveActivityBrowserColumnVisibility() const { return LiveActivityBrowserColumnVisibility; }
+	void SetLiveActivityBrowserColumnVisibility(FColumnVisibilitySnapshot NewValue) { LiveActivityBrowserColumnVisibility = MoveTemp(NewValue); OnLiveActivityBrowserColumnVisibilityEvent.Broadcast(LiveActivityBrowserColumnVisibility); }
+	FOnColumnVisibilitySnapshotChanged& OnLiveActivityBrowserColumnVisibility() { return OnLiveActivityBrowserColumnVisibilityEvent; }
+
+	const FColumnVisibilitySnapshot& GetLiveSessionContentColumnVisibility() const { return LiveSessionContentColumnVisibility; }
+	void SetLiveSessionContentColumnVisibility(FColumnVisibilitySnapshot NewValue) { LiveSessionContentColumnVisibility = MoveTemp(NewValue); OnLiveSessionContentColumnVisibilityEvent.Broadcast(LiveSessionContentColumnVisibility); }
+	FOnColumnVisibilitySnapshotChanged& OnLiveSessionContentColumnVisibility() { return OnLiveSessionContentColumnVisibilityEvent; }
+	
+private:
+	
+	UPROPERTY(Config)
+	FColumnVisibilitySnapshot SessionBrowserColumnVisibility;
+	FOnColumnVisibilitySnapshotChanged OnSessionBrowserColumnVisibilityChangedEvent;
+
+	UPROPERTY(Config)
+	FColumnVisibilitySnapshot ArchivedActivityBrowserColumnVisibility;
+	FOnColumnVisibilitySnapshotChanged OnArchivedActivityBrowserColumnVisibilityEvent;
+
+	UPROPERTY(Config)
+	FColumnVisibilitySnapshot LiveActivityBrowserColumnVisibility;
+	FOnColumnVisibilitySnapshotChanged OnLiveActivityBrowserColumnVisibilityEvent;
+
+	UPROPERTY(Config)
+	FColumnVisibilitySnapshot LiveSessionContentColumnVisibility;
+	FOnColumnVisibilitySnapshotChanged OnLiveSessionContentColumnVisibilityEvent;
+};
diff --git a/Engine/Plugins/Developer/Concert/ConcertApp/MultiUserServer/Source/MultiUserServer/Private/Widgets/Browser/SConcertServerSessionBrowser.cpp b/Engine/Plugins/Developer/Concert/ConcertApp/MultiUserServer/Source/MultiUserServer/Private/Widgets/Browser/SConcertServerSessionBrowser.cpp
--- a/Engine/Plugins/Developer/Concert/ConcertApp/MultiUserServer/Source/MultiUserServer/Private/Widgets/Browser/SConcertServerSessionBrowser.cpp
+++ b/Engine/Plugins/Developer/Concert/ConcertApp/MultiUserServer/Source/MultiUserServer/Private/Widgets/Browser/SConcertServerSessionBrowser.cpp
@@ -2,6 +2,8 @@
 
 #include "SConcertServerSessionBrowser.h"
 
+#include "MultiUserServerUserSettings.h"
+
 #include "Session/Browser/SConcertSessionBrowser.h"
 #include "Widgets/Browser/ConcertServerSessionBrowserController.h"
 #include "Widgets/StatusBar/SConcertStatusBar.h"
@@ -60,7 +62,12 @@ TSharedRef<SWidget> SConcertServerSessionBrowser::MakeSessionTableView(const FAr
 		.OnRequestedDeleteSession(this, &SConcertServerSessionBrowser::RequestDeleteSession)
 		// Pretend a modal dialog said no - RequestDeleteSession will show non-modal dialog
 		.CanDeleteArchivedSession_Lambda([](TSharedPtr<FConcertSessionItem>) { return false; })
-		.CanDeleteActiveSession_Lambda([](TSharedPtr<FConcertSessionItem>) { return false; });
+		.CanDeleteActiveSession_Lambda([](TSharedPtr<FConcertSessionItem>) { return false; })
+		.ColumnVisibilitySnapshot(UMultiUserServerUserSettings::GetUserSettings()->GetSessionBrowserColumnVisibility())
+		.SaveColumnVisibilitySnapshot_Lambda([](const FColumnVisibilitySnapshot& Snapshot)
+		{
+			UMultiUserServerUserSettings::GetUserSettings()->SetSessionBrowserColumnVisibility(Snapshot);
+		});
 }
 
 void SConcertServerSessionBrowser::RequestDeleteSession(const TSharedPtr<FConcertSessionItem>& SessionItem) 
diff --git a/Engine/Plugins/Developer/Concert/ConcertApp/MultiUserServer/Source/MultiUserServer/Private/Widgets/SessionTabs/Archived/ArchivedSessionHistoryController.cpp b/Engine/Plugins/Developer/Concert/ConcertApp/MultiUserServer/Source/MultiUserServer/Private/Widgets/SessionTabs/Archived/ArchivedSessionHistoryController.cpp
--- a/Engine/Plugins/Developer/Concert/ConcertApp/MultiUserServer/Source/MultiUserServer/Private/Widgets/SessionTabs/Archived/ArchivedSessionHistoryController.cpp
+++ b/Engine/Plugins/Developer/Concert/ConcertApp/MultiUserServer/Source/MultiUserServer/Private/Widgets/SessionTabs/Archived/ArchivedSessionHistoryController.cpp
@@ -2,14 +2,43 @@
 
 #include "ArchivedSessionHistoryController.h"
 
+#include "MultiUserServerUserSettings.h"
+
+namespace UE::MultiUserServer::Private
+{
+	static SSessionHistory::FArguments MakeArchivedTabSessionHistoryArguments(SSessionHistory::FArguments&& InArgs)
+	{
+		return InArgs
+			.ColumnVisibilitySnapshot(UMultiUserServerUserSettings::GetUserSettings()->GetArchivedActivityBrowserColumnVisibility())
+			.SaveColumnVisibilitySnapshot_Lambda([](const FColumnVisibilitySnapshot& Snapshot)
+			{
+				UMultiUserServerUserSettings::GetUserSettings()->SetArchivedActivityBrowserColumnVisibility(Snapshot);
+			});
+	}
+}
+
 FArchivedSessionHistoryController::FArchivedSessionHistoryController(FGuid SessionId, TSharedRef<IConcertSyncServer> SyncServer, SSessionHistory::FArguments Arguments)
-	: FServerSessionHistoryControllerBase(MoveTemp(SessionId), MoveTemp(Arguments))
+	: FServerSessionHistoryControllerBase(MoveTemp(SessionId), UE::MultiUserServer::Private::MakeArchivedTabSessionHistoryArguments(MoveTemp(Arguments)))
 	, SyncServer(MoveTemp(SyncServer))
 {
 	ReloadActivities();
+	UMultiUserServerUserSettings::GetUserSettings()->OnArchivedActivityBrowserColumnVisibility().AddRaw(this, &FArchivedSessionHistoryController::OnActivityListColumnVisibilitySettingsUpdated);
+}
+
+FArchivedSessionHistoryController::~FArchivedSessionHistoryController()
+{
+	if (UMultiUserServerUserSettings* Settings = UMultiUserServerUserSettings::GetUserSettings(); IsValid(Settings))
+	{
+		Settings->OnArchivedActivityBrowserColumnVisibility().RemoveAll(this);
+	}
 }
 
 TOptional<FConcertSyncSessionDatabaseNonNullPtr> FArchivedSessionHistoryController::GetSessionDatabase(const FGuid& InSessionId) const
 {
 	return SyncServer->GetArchivedSessionDatabase(InSessionId);
 }
+
+void FArchivedSessionHistoryController::OnActivityListColumnVisibilitySettingsUpdated(const FColumnVisibilitySnapshot& NewValue)
+{
+	GetSessionHistory()->OnColumnVisibilitySettingsChanged(NewValue);
+}
diff --git a/Engine/Plugins/Developer/Concert/ConcertApp/MultiUserServer/Source/MultiUserServer/Private/Widgets/SessionTabs/Archived/ArchivedSessionHistoryController.h b/Engine/Plugins/Developer/Concert/ConcertApp/MultiUserServer/Source/MultiUserServer/Private/Widgets/SessionTabs/Archived/ArchivedSessionHistoryController.h
--- a/Engine/Plugins/Developer/Concert/ConcertApp/MultiUserServer/Source/MultiUserServer/Private/Widgets/SessionTabs/Archived/ArchivedSessionHistoryController.h
+++ b/Engine/Plugins/Developer/Concert/ConcertApp/MultiUserServer/Source/MultiUserServer/Private/Widgets/SessionTabs/Archived/ArchivedSessionHistoryController.h
@@ -10,6 +10,7 @@ class FArchivedSessionHistoryController : public FServerSessionHistoryController
 public:
 	
 	FArchivedSessionHistoryController(FGuid SessionId, TSharedRef<IConcertSyncServer> SyncServer, SSessionHistory::FArguments Arguments);
+	virtual ~FArchivedSessionHistoryController() override;
 	
 protected:
 
@@ -20,4 +21,6 @@ protected:
 private:
 
 	const TSharedRef<IConcertSyncServer> SyncServer;
+
+	void OnActivityListColumnVisibilitySettingsUpdated(const FColumnVisibilitySnapshot& NewValue);
 };
diff --git a/Engine/Plugins/Developer/Concert/ConcertApp/MultiUserServer/Source/MultiUserServer/Private/Widgets/SessionTabs/Live/LiveServerSessionHistoryController.cpp b/Engine/Plugins/Developer/Concert/ConcertApp/MultiUserServer/Source/MultiUserServer/Private/Widgets/SessionTabs/Live/LiveServerSessionHistoryController.cpp
--- a/Engine/Plugins/Developer/Concert/ConcertApp/MultiUserServer/Source/MultiUserServer/Private/Widgets/SessionTabs/Live/LiveServerSessionHistoryController.cpp
+++ b/Engine/Plugins/Developer/Concert/ConcertApp/MultiUserServer/Source/MultiUserServer/Private/Widgets/SessionTabs/Live/LiveServerSessionHistoryController.cpp
@@ -5,12 +5,27 @@
 #include "ConcertSyncSessionDatabase.h"
 #include "IConcertSession.h"
 #include "IConcertSyncServer.h"
+#include "MultiUserServerUserSettings.h"
+
+namespace UE::MultiUserServer::Private
+{
+	static SSessionHistory::FArguments MakeLiveTabSessionHistoryArguments()
+	{
+		return SSessionHistory::FArguments()
+			.ColumnVisibilitySnapshot(UMultiUserServerUserSettings::GetUserSettings()->GetLiveActivityBrowserColumnVisibility())
+			.SaveColumnVisibilitySnapshot_Lambda([](const FColumnVisibilitySnapshot& Snapshot)
+			{
+				UMultiUserServerUserSettings::GetUserSettings()->SetLiveActivityBrowserColumnVisibility(Snapshot);
+			});
+	}
+}
 
 FLiveServerSessionHistoryController::FLiveServerSessionHistoryController(TSharedRef<IConcertServerSession> InspectedSession, TSharedRef<IConcertSyncServer> SyncServer)
-	: FServerSessionHistoryControllerBase(InspectedSession->GetId())
+	: FServerSessionHistoryControllerBase(InspectedSession->GetId(), UE::MultiUserServer::Private::MakeLiveTabSessionHistoryArguments())
 	, SyncServer(MoveTemp(SyncServer))
 {
 	ReloadActivities();
+	UMultiUserServerUserSettings::GetUserSettings()->OnLiveActivityBrowserColumnVisibility().AddRaw(this, &FLiveServerSessionHistoryController::OnActivityListColumnVisibilitySettingsUpdated);
 
 	if (const TOptional<FConcertSyncSessionDatabaseNonNullPtr> Database = SyncServer->GetLiveSessionDatabase(GetSessionId()))
 	{
@@ -20,6 +35,10 @@ FLiveServerSessionHistoryController::FLiveServerSessionHistoryController(TShared
 
 FLiveServerSessionHistoryController::~FLiveServerSessionHistoryController()
 {
+	if (UMultiUserServerUserSettings* Settings = UMultiUserServerUserSettings::GetUserSettings(); IsValid(Settings))
+	{
+		Settings->OnLiveActivityBrowserColumnVisibility().RemoveAll(this);
+	}
 	if (const TOptional<FConcertSyncSessionDatabaseNonNullPtr> Database = SyncServer->GetLiveSessionDatabase(GetSessionId()))
 	{
 		Database->OnActivityProduced().RemoveAll(this);
@@ -35,3 +54,8 @@ void FLiveServerSessionHistoryController::OnSessionProduced(const FConcertSyncAc
 {
 	ReloadActivities();
 }
+
+void FLiveServerSessionHistoryController::OnActivityListColumnVisibilitySettingsUpdated(const FColumnVisibilitySnapshot& NewValue)
+{
+	GetSessionHistory()->OnColumnVisibilitySettingsChanged(NewValue);
+}
diff --git a/Engine/Plugins/Developer/Concert/ConcertApp/MultiUserServer/Source/MultiUserServer/Private/Widgets/SessionTabs/Live/LiveServerSessionHistoryController.h b/Engine/Plugins/Developer/Concert/ConcertApp/MultiUserServer/Source/MultiUserServer/Private/Widgets/SessionTabs/Live/LiveServerSessionHistoryController.h
--- a/Engine/Plugins/Developer/Concert/ConcertApp/MultiUserServer/Source/MultiUserServer/Private/Widgets/SessionTabs/Live/LiveServerSessionHistoryController.h
+++ b/Engine/Plugins/Developer/Concert/ConcertApp/MultiUserServer/Source/MultiUserServer/Private/Widgets/SessionTabs/Live/LiveServerSessionHistoryController.h
@@ -28,4 +28,5 @@ private:
 	TSharedRef<IConcertSyncServer> SyncServer;
 	
 	void OnSessionProduced(const FConcertSyncActivity& ProducedActivity);
+	void OnActivityListColumnVisibilitySettingsUpdated(const FColumnVisibilitySnapshot& NewValue);
 };
diff --git a/Engine/Plugins/Developer/Concert/ConcertApp/MultiUserServer/Source/MultiUserServer/Private/Widgets/SessionTabs/Live/PackageViewer/ConcertSessionPackageViewerController.h b/Engine/Plugins/Developer/Concert/ConcertApp/MultiUserServer/Source/MultiUserServer/Private/Widgets/SessionTabs/Live/PackageViewer/ConcertSessionPackageViewerController.h
--- a/Engine/Plugins/Developer/Concert/ConcertApp/MultiUserServer/Source/MultiUserServer/Private/Widgets/SessionTabs/Live/PackageViewer/ConcertSessionPackageViewerController.h
+++ b/Engine/Plugins/Developer/Concert/ConcertApp/MultiUserServer/Source/MultiUserServer/Private/Widgets/SessionTabs/Live/PackageViewer/ConcertSessionPackageViewerController.h
@@ -4,16 +4,18 @@
 
 #include "CoreMinimal.h"
 
-struct FConcertSyncActivity;
-struct FConcertClientInfo;
-struct FConcertSyncPackageEventMetaData;
-struct FConcertSessionActivity;
 class IConcertSyncServer;
 class IConcertServerSession;
 class SConcertSessionPackageViewer;
 
+struct FConcertClientInfo;
+struct FConcertSessionActivity;
+struct FConcertSyncActivity;
+struct FConcertSyncPackageEventMetaData;
+struct FColumnVisibilitySnapshot;
+
 /** Manages SConcertSessionPackageViewer for a server. */
-class FConcertSessionPackageViewerController
+class FConcertSessionPackageViewerController : public TSharedFromThis<FConcertSessionPackageViewerController>
 {
 public:
 
@@ -37,4 +39,6 @@ private:
 	bool GetPackageEvent(const FConcertSessionActivity& Activity, FConcertSyncPackageEventMetaData& OutPackageEvent) const;
 	TOptional<int64> GetSizeOfPackageEvent(const FConcertSessionActivity& Activity) const;
 	void OnSessionProduced(const FConcertSyncActivity& ProducedActivity) const;
+	
+	void OnSessionContentColumnVisibilitySettingsUpdated(const FColumnVisibilitySnapshot& NewValue);
 };
diff --git a/Engine/Plugins/Developer/Concert/ConcertApp/MultiUserServer/Source/MultiUserServer/Private/Widgets/SessionTabs/Live/PackageViewer/PackageViewerColumns.cpp b/Engine/Plugins/Developer/Concert/ConcertApp/MultiUserServer/Source/MultiUserServer/Private/Widgets/SessionTabs/Live/PackageViewer/PackageViewerColumns.cpp
--- a/Engine/Plugins/Developer/Concert/ConcertApp/MultiUserServer/Source/MultiUserServer/Private/Widgets/SessionTabs/Live/PackageViewer/PackageViewerColumns.cpp
+++ b/Engine/Plugins/Developer/Concert/ConcertApp/MultiUserServer/Source/MultiUserServer/Private/Widgets/SessionTabs/Live/PackageViewer/PackageViewerColumns.cpp
@@ -25,7 +25,7 @@ namespace UE::MultiUserServer::PackageViewerColumns
 		return FActivityColumn(
 			SHeaderRow::Column(ColumnId)
 				.DefaultLabel(Label)
-				.ManualWidth(Width)
+				.FillSized(Width)
 			)
 			.ColumnSortOrder(static_cast<int32>(SortOrder))
 			.GenerateColumnWidget(FActivityColumn::FGenerateColumnWidget::CreateLambda([GetterFunc, FormatFunc](const TSharedRef<SConcertSessionActivities>& Owner, const TSharedRef<FConcertSessionActivity>& Activity, SOverlay::FScopedWidgetSlotArguments& Slot)
diff --git a/Engine/Plugins/Developer/Concert/ConcertApp/MultiUserServer/Source/MultiUserServer/Private/Widgets/SessionTabs/Live/PackageViewer/SConcertSessionPackageViewer.cpp b/Engine/Plugins/Developer/Concert/ConcertApp/MultiUserServer/Source/MultiUserServer/Private/Widgets/SessionTabs/Live/PackageViewer/SConcertSessionPackageViewer.cpp
--- a/Engine/Plugins/Developer/Concert/ConcertApp/MultiUserServer/Source/MultiUserServer/Private/Widgets/SessionTabs/Live/PackageViewer/SConcertSessionPackageViewer.cpp
+++ b/Engine/Plugins/Developer/Concert/ConcertApp/MultiUserServer/Source/MultiUserServer/Private/Widgets/SessionTabs/Live/PackageViewer/SConcertSessionPackageViewer.cpp
@@ -2,6 +2,7 @@
 
 #include "SConcertSessionPackageViewer.h"
 
+#include "MultiUserServerUserSettings.h"
 #include "PackageViewerColumns.h"
 
 #include "Session/Activity/PredefinedActivityColumns.h"
@@ -37,7 +38,12 @@ void SConcertSessionPackageViewer::Construct(const FArguments& InArgs)
 		.PackageActivitiesVisibility(EVisibility::Visible)
 		.TransactionActivitiesVisibility(EVisibility::Collapsed)
 		.DetailsAreaVisibility(EVisibility::Collapsed)
-		.IsAutoScrollEnabled(true);
+		.IsAutoScrollEnabled(true)
+		.ColumnVisibilitySnapshot(UMultiUserServerUserSettings::GetUserSettings()->GetLiveSessionContentColumnVisibility())
+		.SaveColumnVisibilitySnapshot_Lambda([](const FColumnVisibilitySnapshot& Snapshot)
+		{
+			UMultiUserServerUserSettings::GetUserSettings()->SetLiveSessionContentColumnVisibility(Snapshot);
+		});
 
 	ChildSlot
 	[
@@ -63,8 +69,12 @@ void SConcertSessionPackageViewer::Construct(const FArguments& InArgs)
 		[
 			ActivityListViewOptions->MakeStatusBar(
 				TAttribute<int32>(ActivityListView.Get(), &SConcertSessionActivities::GetTotalActivityNum),
-				TAttribute<int32>(ActivityListView.Get(), &SConcertSessionActivities::GetDisplayedActivityNum)
-				)
+				TAttribute<int32>(ActivityListView.Get(), &SConcertSessionActivities::GetDisplayedActivityNum),
+				FConcertSessionActivitiesOptions::FExtendContextMenu::CreateLambda([this](FMenuBuilder& MenuBuilder)
+				{
+					MenuBuilder.AddSeparator();
+					AddEntriesForShowingHiddenRows(ActivityListView->GetHeaderRow().ToSharedRef(), MenuBuilder);
+				}))
 		]
 	];
 }
@@ -80,6 +90,11 @@ void SConcertSessionPackageViewer::AppendActivity(FConcertSessionActivity Activi
 	ActivityListView->Append(NewActivity);
 }
 
+void SConcertSessionPackageViewer::OnColumnVisibilitySettingsChanged(const FColumnVisibilitySnapshot& ColumnSnapshot)
+{
+	ActivityListView->OnColumnVisibilitySettingsChanged(ColumnSnapshot);
+}
+
 TOptional<int64> SConcertSessionPackageViewer::GetVersionOfPackageActivity(const FConcertSessionActivity& Activity, SConcertSessionActivities::FGetPackageEvent GetPackageEventFunc) const
 {
 	FConcertSyncPackageEventMetaData PackageEventMetaData;
diff --git a/Engine/Plugins/Developer/Concert/ConcertApp/MultiUserServer/Source/MultiUserServer/Private/Widgets/SessionTabs/Live/PackageViewer/SConcertSessionPackageViewer.h b/Engine/Plugins/Developer/Concert/ConcertApp/MultiUserServer/Source/MultiUserServer/Private/Widgets/SessionTabs/Live/PackageViewer/SConcertSessionPackageViewer.h
--- a/Engine/Plugins/Developer/Concert/ConcertApp/MultiUserServer/Source/MultiUserServer/Private/Widgets/SessionTabs/Live/PackageViewer/SConcertSessionPackageViewer.h
+++ b/Engine/Plugins/Developer/Concert/ConcertApp/MultiUserServer/Source/MultiUserServer/Private/Widgets/SessionTabs/Live/PackageViewer/SConcertSessionPackageViewer.h
@@ -28,6 +28,8 @@ public:
 	void ResetActivityList();
 	void AppendActivity(FConcertSessionActivity Activity);
 
+	void OnColumnVisibilitySettingsChanged(const FColumnVisibilitySnapshot& ColumnSnapshot);
+
 private:
 
 	TSharedPtr<SConcertSessionActivities> ActivityListView;
diff --git a/Engine/Plugins/Developer/Concert/ConcertUI/ConcertSharedSlate/Source/ConcertSharedSlate/ConcertSharedSlate.Build.cs b/Engine/Plugins/Developer/Concert/ConcertUI/ConcertSharedSlate/Source/ConcertSharedSlate/ConcertSharedSlate.Build.cs
--- a/Engine/Plugins/Developer/Concert/ConcertUI/ConcertSharedSlate/Source/ConcertSharedSlate/ConcertSharedSlate.Build.cs
+++ b/Engine/Plugins/Developer/Concert/ConcertUI/ConcertSharedSlate/Source/ConcertSharedSlate/ConcertSharedSlate.Build.cs
@@ -24,6 +24,7 @@ namespace UnrealBuildTool.Rules
 					"ConcertTransport",
 					"ConcertSyncCore",
 					"InputCore",
+					"Json",
 					"Projects",
 					"SlateCore",
 					"UndoHistory",
diff --git a/Engine/Plugins/Developer/Concert/ConcertUI/ConcertSharedSlate/Source/ConcertSharedSlate/Private/ConcertHeaderRowUtils.cpp b/Engine/Plugins/Developer/Concert/ConcertUI/ConcertSharedSlate/Source/ConcertSharedSlate/Private/ConcertHeaderRowUtils.cpp
new file mode 100644
--- /dev/null
+++ b/Engine/Plugins/Developer/Concert/ConcertUI/ConcertSharedSlate/Source/ConcertSharedSlate/Private/ConcertHeaderRowUtils.cpp
@@ -0,0 +1,128 @@
+// Copyright Epic Games, Inc. All Rights Reserved.
+
+#include "ConcertHeaderRowUtils.h"
+
+#include "Algo/AllOf.h"
+#include "Algo/Find.h"
+
+#include "Dom/JsonObject.h"
+
+#include "Framework/MultiBox/MultiBoxBuilder.h"
+
+#include "Serialization/BufferArchive.h"
+#include "Serialization/JsonSerializer.h"
+
+#include "Widgets/Views/SHeaderRow.h"
+
+#define LOCTEXT_NAMESPACE "ConcertHeaderRowUtils"
+
+namespace UE::ConcertSharedSlate
+{
+	TSharedRef<SWidget> MakeHideColumnContextMenu(const TSharedRef<SHeaderRow>& HeaderRow, const FName ForColumnID)
+	{
+		constexpr bool bInShouldCloseWindowAfterMenuSelection = true;
+		FMenuBuilder MenuBuilder(bInShouldCloseWindowAfterMenuSelection, nullptr);
+		AddHideColumnEntry(HeaderRow, ForColumnID, MenuBuilder);
+		return MenuBuilder.MakeWidget();
+	}
+
+	void AddHideColumnEntry(const TSharedRef<SHeaderRow>& HeaderRow, const FName ForColumnID, FMenuBuilder& MenuBuilder)
+	{
+		const SHeaderRow::FColumn* FoundColumn = Algo::FindByPredicate(HeaderRow->GetColumns(), [ForColumnID](const SHeaderRow::FColumn& Column){ return  Column.ColumnId == ForColumnID; });
+		check(FoundColumn);
+
+		TWeakPtr<SHeaderRow> WeakHeaderRow = HeaderRow;
+		MenuBuilder.AddMenuEntry(
+			LOCTEXT("HideColumn", "Hide column"),
+			FText::Format(LOCTEXT("HideColumn_Tooltip", "Hides the {0} column. You can unhide it using the eye-icon."), FoundColumn->DefaultText.Get()),
+			FSlateIcon(),
+			FUIAction(
+				FExecuteAction::CreateLambda([WeakHeaderRow, ForColumnID](){ WeakHeaderRow.Pin()->SetShowGeneratedColumn(ForColumnID, !WeakHeaderRow.Pin()->IsColumnVisible(ForColumnID)); }),
+				FCanExecuteAction::CreateLambda([] { return true; })),
+			NAME_None,
+			EUserInterfaceActionType::Button
+			);
+	}
+
+	void AddEntriesForShowingHiddenRows(const TSharedRef<SHeaderRow>& HeaderRow, FMenuBuilder& MenuBuilder)
+	{
+		for (int32 ColumnIndex = 0; ColumnIndex < HeaderRow->GetColumns().Num(); ++ColumnIndex)
+		{
+			const SHeaderRow::FColumn& Column = HeaderRow->GetColumns()[ColumnIndex];
+			TWeakPtr<SHeaderRow> WeakHeaderRow = HeaderRow;
+			if (!Column.ShouldGenerateWidget.IsSet())
+			{
+				MenuBuilder.AddMenuEntry(
+					FText::Format(LOCTEXT("UnhideColumn", "Show \"{0}\" column"), Column.DefaultText.Get()),
+					FText::Format(LOCTEXT("UnhideColumn_Tooltip", "Show the {0} column."), Column.DefaultText.Get()),
+					FSlateIcon(),
+					FUIAction(
+						FExecuteAction::CreateLambda([WeakHeaderRow, ColumnID = Column.ColumnId, ColumnIndex](){ WeakHeaderRow.Pin()->SetShowGeneratedColumn(ColumnID, !WeakHeaderRow.Pin()->GetColumns()[ColumnIndex].bIsVisible); }),
+						FCanExecuteAction::CreateLambda([] { return true; }),
+						FIsActionChecked::CreateLambda([WeakHeaderRow, ColumnIndex] { return WeakHeaderRow.Pin()->GetColumns()[ColumnIndex].bIsVisible; })
+						),
+					NAME_None,
+					EUserInterfaceActionType::ToggleButton
+				);
+			}
+		}
+	}
+
+	static const FString JsonField_ColumnId("ColumnID");
+	static const FString JsonField_IsVisible("bIsVisible");
+	static const FString JsonField_Values("Values");
+
+	FColumnVisibilitySnapshot SnapshotColumnVisibilityState(const TSharedRef<SHeaderRow>& HeaderRow)
+	{
+		TSharedPtr<FJsonObject> JsonObject = MakeShared<FJsonObject>();
+
+		TArray<TSharedPtr<FJsonValue>> Array;
+		for (const SHeaderRow::FColumn& Column : HeaderRow->GetColumns())
+		{
+			TSharedPtr<FJsonObject> ColumnObject = MakeShared<FJsonObject>();
+			ColumnObject->SetStringField(JsonField_ColumnId, Column.ColumnId.ToString());
+			ColumnObject->SetBoolField(JsonField_IsVisible, Column.bIsVisible);
+			Array.Add(MakeShared<FJsonValueObject>(ColumnObject));
+		}
+		JsonObject->SetArrayField(JsonField_Values, Array);
+
+		FString Snapshot;
+		const TSharedRef<TJsonWriter<>> Writer = TJsonWriterFactory<>::Create(&Snapshot);
+		FJsonSerializer::Serialize(JsonObject.ToSharedRef(), Writer);
+		return { Snapshot };
+	}
+
+	void RestoreColumnVisibilityState(const TSharedRef<SHeaderRow>& HeaderRow, const FColumnVisibilitySnapshot& Snapshot)
+	{
+		TSharedPtr<FJsonObject> JsonObject = MakeShared<FJsonObject>();
+		const TSharedRef<TJsonReader<>> Reader = TJsonReaderFactory<>::Create(Snapshot.Snapshot);
+		const bool bIsValid = FJsonSerializer::Deserialize(Reader, JsonObject);
+		// This is a legal state, e.g. trying to restore a column view for the first time when no data is yet available
+		if (!bIsValid)
+		{
+			return;
+		}
+
+		const TArray<TSharedPtr<FJsonValue>> ColumnObjects = JsonObject->GetArrayField(JsonField_Values);
+		for (const TSharedPtr<FJsonValue>& ColumnValue : ColumnObjects)
+		{
+			TSharedPtr<FJsonObject>* PossibleColumnObject;
+			if (!ColumnValue->TryGetObject(PossibleColumnObject))
+			{
+				continue;
+			}
+			const TSharedPtr<FJsonObject>& ColumnObject = *PossibleColumnObject;
+			
+			FString ColumnId;
+			bool bIsVisible;
+			if (!ColumnObject->TryGetStringField(JsonField_ColumnId, ColumnId) || !ColumnObject->TryGetBoolField(JsonField_IsVisible, bIsVisible))
+			{
+				continue;
+			}
+
+			HeaderRow->SetShowGeneratedColumn(*ColumnId, bIsVisible);
+		}
+	}
+}
+
+#undef LOCTEXT_NAMESPACE
\ No newline at end of file
diff --git a/Engine/Plugins/Developer/Concert/ConcertUI/ConcertSharedSlate/Source/ConcertSharedSlate/Private/Session/Activity/PredefinedActivityColumns.cpp b/Engine/Plugins/Developer/Concert/ConcertUI/ConcertSharedSlate/Source/ConcertSharedSlate/Private/Session/Activity/PredefinedActivityColumns.cpp
--- a/Engine/Plugins/Developer/Concert/ConcertUI/ConcertSharedSlate/Source/ConcertSharedSlate/Private/Session/Activity/PredefinedActivityColumns.cpp
+++ b/Engine/Plugins/Developer/Concert/ConcertUI/ConcertSharedSlate/Source/ConcertSharedSlate/Private/Session/Activity/PredefinedActivityColumns.cpp
@@ -37,7 +37,7 @@ FActivityColumn UE::ConcertSharedSlate::ActivityColumn::DateTime()
 	return FActivityColumn(
 		SHeaderRow::Column(DateTimeColumnId)
 			.DefaultLabel(LOCTEXT("DateTime", "Date/Time"))
-			.ManualWidth(160)
+			.FillSized(160)
 		)
 		.ColumnSortOrder(static_cast<int32>(EPredefinedColumnOrder::DateTime))
 		.GenerateColumnWidget(FActivityColumn::FGenerateColumnWidget::CreateLambda([](const TSharedRef<SConcertSessionActivities>& Owner, const TSharedRef<FConcertSessionActivity>& Activity, SOverlay::FScopedWidgetSlotArguments& Slot)
@@ -68,6 +68,7 @@ FActivityColumn UE::ConcertSharedSlate::ActivityColumn::Summary()
 	return FActivityColumn(
 		SHeaderRow::Column(SummaryColumnId)
 			.DefaultLabel(LOCTEXT("Summary", "Summary"))
+			.ShouldGenerateWidget(true)
 		)
 		.GenerateColumnWidget(FActivityColumn::FGenerateColumnWidget::CreateLambda([](const TSharedRef<SConcertSessionActivities>& Owner, const TSharedRef<FConcertSessionActivity>& Activity, SOverlay::FScopedWidgetSlotArguments& Slot)
 		{
@@ -103,6 +104,7 @@ FActivityColumn UE::ConcertSharedSlate::ActivityColumn::AvatarColor()
 		SHeaderRow::Column(AvatarColorColumnId)
 			.DefaultLabel(INVTEXT(""))
 			.ManualWidth(8)
+			.ShouldGenerateWidget(true)
 		)
 		.ColumnSortOrder(static_cast<int32>(EPredefinedColumnOrder::AvatarColor))
 		.GenerateColumnWidget(FActivityColumn::FGenerateColumnWidget::CreateLambda([](const TSharedRef<SConcertSessionActivities>& Owner, const TSharedRef<FConcertSessionActivity>& Activity, SOverlay::FScopedWidgetSlotArguments& Slot)
@@ -126,7 +128,7 @@ FActivityColumn UE::ConcertSharedSlate::ActivityColumn::ClientName()
 	return FActivityColumn(
 		SHeaderRow::Column(ClientNameColumnId)
 			.DefaultLabel(LOCTEXT("Client", "Client"))
-			.ManualWidth(80)
+			.FillSized(80)
 		)
 		.ColumnSortOrder(static_cast<int32>(EPredefinedColumnOrder::ClientName))
 		.GenerateColumnWidget(FActivityColumn::FGenerateColumnWidget::CreateLambda([](const TSharedRef<SConcertSessionActivities>& Owner, const TSharedRef<FConcertSessionActivity>& Activity, SOverlay::FScopedWidgetSlotArguments& Slot)
@@ -157,7 +159,7 @@ FActivityColumn UE::ConcertSharedSlate::ActivityColumn::Operation()
 	return FActivityColumn(
 		SHeaderRow::Column(OperationColumnId)
 			.DefaultLabel(LOCTEXT("Operation", "Operation"))
-			.ManualWidth(160)
+			.FillSized(160)
 		)
 		.ColumnSortOrder(static_cast<int32>(EPredefinedColumnOrder::Operation))
 		.GenerateColumnWidget(FActivityColumn::FGenerateColumnWidget::CreateLambda([](const TSharedRef<SConcertSessionActivities>& Owner, const TSharedRef<FConcertSessionActivity>& Activity, SOverlay::FScopedWidgetSlotArguments& Slot)
@@ -181,7 +183,7 @@ FActivityColumn UE::ConcertSharedSlate::ActivityColumn::Package()
 	return FActivityColumn(
 		SHeaderRow::Column(PackageColumnId)
 			.DefaultLabel(LOCTEXT("Package", "Package"))
-			.ManualWidth(200)
+			.FillSized(200)
 		)
 		.ColumnSortOrder(static_cast<int32>(EPredefinedColumnOrder::Package))
 		.GenerateColumnWidget(FActivityColumn::FGenerateColumnWidget::CreateLambda([](const TSharedRef<SConcertSessionActivities>& Owner, const TSharedRef<FConcertSessionActivity>& Activity, SOverlay::FScopedWidgetSlotArguments& Slot)
diff --git a/Engine/Plugins/Developer/Concert/ConcertUI/ConcertSharedSlate/Source/ConcertSharedSlate/Private/Session/Activity/SConcertSessionActivities.cpp b/Engine/Plugins/Developer/Concert/ConcertUI/ConcertSharedSlate/Source/ConcertSharedSlate/Private/Session/Activity/SConcertSessionActivities.cpp
--- a/Engine/Plugins/Developer/Concert/ConcertUI/ConcertSharedSlate/Source/ConcertSharedSlate/Private/Session/Activity/SConcertSessionActivities.cpp
+++ b/Engine/Plugins/Developer/Concert/ConcertUI/ConcertSharedSlate/Source/ConcertSharedSlate/Private/Session/Activity/SConcertSessionActivities.cpp
@@ -365,7 +365,17 @@
 	checkf(!bContainsSummary, TEXT("Summary is already created by SConcertSessionActivities!"))
 	Columns.Add(ActivityColumn::Summary());
 	
-	HeaderRow = SNew(SHeaderRow);
+	HeaderRow = SNew(SHeaderRow)
+		.OnHiddenColumnsListChanged_Lambda([this, SaveCallback = InArgs._SaveColumnVisibilitySnapshot]()
+		{
+			// bIsUpdatingColumnVisibility is true when the columns visibility is being changed by OnColumnVisibilitySettingsChanged to prevent recursion
+			if (SaveCallback.IsBound() && !bIsUpdatingColumnVisibility)
+			{
+				const FColumnVisibilitySnapshot Snapshot = SnapshotColumnVisibilityState(HeaderRow.ToSharedRef());
+				SaveCallback.Execute(Snapshot);
+			}
+		});
+	
 	TSet<FName> DuplicateColumnDetection;
 	for (FActivityColumn& Column : Columns)
 	{
@@ -374,9 +384,16 @@
 		
 		// SHeaderRow owns the columns and deletes them when destroyed
 		FActivityColumn* ManagedByHeaderRow = new FActivityColumn(MoveTemp(Column));
+		const bool bCanBeHidden = !ManagedByHeaderRow->ShouldGenerateWidget.IsSet(); 
+		if (bCanBeHidden)
+		{
+			ManagedByHeaderRow->OnGetMenuContent = FOnGetContent::CreateLambda([this, ColumnId = Column.ColumnId](){ return MakeHideColumnContextMenu(HeaderRow.ToSharedRef(), ColumnId); });
+		}
 		HeaderRow->AddColumn(*ManagedByHeaderRow);
 	}
-	
+
+	TGuardValue<bool> GuardValue(bIsUpdatingColumnVisibility, true);
+	RestoreColumnVisibilityState(HeaderRow.ToSharedRef(), InArgs._ColumnVisibilitySnapshot);
 	return HeaderRow.ToSharedRef();
 }
 
@@ -735,6 +752,12 @@
 	return SearchTextFilter->GetFilterErrorText();
 }
 
+void SConcertSessionActivities::OnColumnVisibilitySettingsChanged(const FColumnVisibilitySnapshot& ColumnSnapshot)
+{
+	TGuardValue<bool> GuardValue(bIsUpdatingColumnVisibility, true);
+	UE::ConcertSharedSlate::RestoreColumnVisibilityState(GetHeaderRow().ToSharedRef(), ColumnSnapshot);
+}
+
 void SConcertSessionActivities::PopulateSearchStrings(const FConcertSessionActivity& Activity, TArray<FString>& OutSearchStrings)
 {
 	for (const SHeaderRow::FColumn& Column : HeaderRow->GetColumns())
@@ -822,7 +845,80 @@
 	LoadingDetailsPanel->SetVisibility(VisiblePanel == LoadingDetailsPanel.Get() ? EVisibility::Visible : EVisibility::Collapsed);
 }
 
-TSharedRef<SWidget> FConcertSessionActivitiesOptions::MakeMenuWidget()
+TSharedRef<SWidget> FConcertSessionActivitiesOptions::MakeStatusBar(TAttribute<int32> Total, TAttribute<int32> Displayed, TOptional<FExtendContextMenu> ExtendMenu)
 {
+	return SNew(SHorizontalBox)
+
+	// Operation count.
+	+SHorizontalBox::Slot()
+	.AutoWidth()
+	.VAlign(VAlign_Center)
+	[
+		MakeDisplayedActivityCountWidget(MoveTemp(Total), MoveTemp(Displayed))
+	]
+
+	// Gap Filler
+	+SHorizontalBox::Slot()
+	.FillWidth(1.0)
+	[
+		SNew(SSpacer)
+	]
+
+	// View option
+	+SHorizontalBox::Slot()
+	.AutoWidth()
+	[
+		MakeViewOptionsWidget(ExtendMenu)
+	];
+}
+
+TSharedRef<SWidget> FConcertSessionActivitiesOptions::MakeDisplayedActivityCountWidget(TAttribute<int32> Total, TAttribute<int32> Displayed)
+{
+	return SNew(STextBlock)
+		.Text_Lambda([Total = MoveTemp(Total), Displayed = MoveTemp(Displayed)]()
+		{
+			if (Total.Get() == Displayed.Get())
+			{
+				return FText::Format(LOCTEXT("OperationCount", "{0} operations"), Total.Get());
+			}
+			else
+			{
+				return FText::Format(LOCTEXT("PartialOperationCount", "Showing {0} of {1} {1}|plural(one=operation,other=operations)"), Displayed.Get(), Total.Get());
+			}
+		});
+}
+
+TSharedRef<SWidget> FConcertSessionActivitiesOptions::MakeViewOptionsWidget(TOptional<FExtendContextMenu> ExtendMenu)
+{
+	return SNew(SComboButton)
+	.ComboButtonStyle(FAppStyle::Get(), "GenericFilters.ComboButtonStyle")
+	.ForegroundColor(FLinearColor::White)
+	.ContentPadding(0)
+	.OnGetMenuContent(this, &FConcertSessionActivitiesOptions::MakeMenuWidget, MoveTemp(ExtendMenu))
+	.HasDownArrow(true)
+	.ContentPadding(FMargin(1, 0))
+	.ButtonContent()
+	[
+		SNew(SHorizontalBox)
+
+		+SHorizontalBox::Slot()
+		.AutoWidth()
+		.VAlign(VAlign_Center)
+		[
+			SNew(SImage).Image(FAppStyle::Get().GetBrush("GenericViewButton")) // The eye ball image.
+		]
+
+		+SHorizontalBox::Slot()
+		.AutoWidth()
+		.Padding(2, 0, 0, 0)
+		.VAlign(VAlign_Center)
+		[
+			SNew(STextBlock).Text(LOCTEXT("ViewOptions", "View Options"))
+		]
+	];
+}
+
+TSharedRef<SWidget> FConcertSessionActivitiesOptions::MakeMenuWidget(TOptional<FExtendContextMenu> ExtendMenu)
+{
 	FMenuBuilder MenuBuilder(/*bInShouldCloseWindowAfterMenuSelection=*/true, nullptr);
 
@@ -912,83 +1008,15 @@
 			EUserInterfaceActionType::ToggleButton
 		);
 	}
+
+	if (ExtendMenu)
+	{
+		ExtendMenu->Execute(MenuBuilder);
+	}
 	
 	return MenuBuilder.MakeWidget();
 }
 
-TSharedRef<SWidget> FConcertSessionActivitiesOptions::MakeViewOptionsWidget()
-{
-	return SNew(SComboButton)
-	.ComboButtonStyle(FAppStyle::Get(), "GenericFilters.ComboButtonStyle")
-	.ForegroundColor(FLinearColor::White)
-	.ContentPadding(0)
-	.OnGetMenuContent(this, &FConcertSessionActivitiesOptions::MakeMenuWidget)
-	.HasDownArrow(true)
-	.ContentPadding(FMargin(1, 0))
-	.ButtonContent()
-	[
-		SNew(SHorizontalBox)
-
-		+SHorizontalBox::Slot()
-		.AutoWidth()
-		.VAlign(VAlign_Center)
-		[
-			SNew(SImage).Image(FAppStyle::Get().GetBrush("GenericViewButton")) // The eye ball image.
-		]
-
-		+SHorizontalBox::Slot()
-		.AutoWidth()
-		.Padding(2, 0, 0, 0)
-		.VAlign(VAlign_Center)
-		[
-			SNew(STextBlock).Text(LOCTEXT("ViewOptions", "View Options"))
-		]
-	];
-}
-
-TSharedRef<SWidget> FConcertSessionActivitiesOptions::MakeDisplayedActivityCountWidget(TAttribute<int32> Total, TAttribute<int32> Displayed)
-{
-	return SNew(STextBlock)
-		.Text_Lambda([Total = MoveTemp(Total), Displayed = MoveTemp(Displayed)]()
-		{
-			if (Total.Get() == Displayed.Get())
-			{
-				return FText::Format(LOCTEXT("OperationCount", "{0} operations"), Total.Get());
-			}
-			else
-			{
-				return FText::Format(LOCTEXT("PartialOperationCount", "Showing {0} of {1} {1}|plural(one=operation,other=operations)"), Displayed.Get(), Total.Get());
-			}
-		});
-}
-
-TSharedRef<SWidget> FConcertSessionActivitiesOptions::MakeStatusBar(TAttribute<int32> Total, TAttribute<int32> Displayed)
-{
-	return SNew(SHorizontalBox)
-
-	// Operation count.
-	+SHorizontalBox::Slot()
-	.AutoWidth()
-	.VAlign(VAlign_Center)
-	[
-		MakeDisplayedActivityCountWidget(MoveTemp(Total), MoveTemp(Displayed))
-	]
-
-	// Gap Filler
-	+SHorizontalBox::Slot()
-	.FillWidth(1.0)
-	[
-		SNew(SSpacer)
-	]
-
-	// View option
-	+SHorizontalBox::Slot()
-	.AutoWidth()
-	[
-		MakeViewOptionsWidget()
-	];
-}
-
 void FConcertSessionActivitiesOptions::OnOptionToggled(const FName CheckBoxId)
 {
 	if (CheckBoxId == ConcertSessionActivityUtils::DisplayRelativeTimeCheckBoxId)
diff --git a/Engine/Plugins/Developer/Concert/ConcertUI/ConcertSharedSlate/Source/ConcertSharedSlate/Private/Session/Browser/SConcertSessionBrowser.cpp b/Engine/Plugins/Developer/Concert/ConcertUI/ConcertSharedSlate/Source/ConcertSharedSlate/Private/Session/Browser/SConcertSessionBrowser.cpp
--- a/Engine/Plugins/Developer/Concert/ConcertUI/ConcertSharedSlate/Source/ConcertSharedSlate/Private/Session/Browser/SConcertSessionBrowser.cpp
+++ b/Engine/Plugins/Developer/Concert/ConcertUI/ConcertSharedSlate/Source/ConcertSharedSlate/Private/Session/Browser/SConcertSessionBrowser.cpp
@@ -4,6 +4,10 @@
 
 #include "Algo/ForEach.h"
 #include "ConcertFrontendStyle.h"
+#include "ConcertHeaderRowUtils.h"
+
+#include "Algo/Find.h"
+
 #include "Session/Browser/ConcertBrowserUtils.h"
 #include "Session/Browser/ConcertSessionBrowserSettings.h"
 #include "Session/Browser/ConcertSessionItem.h"
@@ -74,7 +78,7 @@ void SConcertSessionBrowser::Construct(const FArguments& InArgs, TSharedRef<ICon
 
 TSharedRef<SWidget> SConcertSessionBrowser::MakeBrowserContent(const FArguments& InArgs)
 {
-	TSharedRef<SWidget> SessionTable = MakeSessionTableView();
+	TSharedRef<SWidget> SessionTable = MakeSessionTableView(InArgs);
 	return SNew(SBorder)
 		.BorderImage(FAppStyle::Get().GetBrush("ToolPanel.GroupBorder"))
 		.Padding(FMargin(1.0f, 2.0f))
@@ -339,14 +343,16 @@ TSharedRef<SWidget> SConcertSessionBrowser::MakeButtonBar(const FArguments& InAr
 	return RowBuilder.MakeWidget();
 }
 
-TSharedRef<SWidget> SConcertSessionBrowser::MakeSessionTableView()
+TSharedRef<SWidget> SConcertSessionBrowser::MakeSessionTableView(const FArguments& InArgs)
 {
+	using namespace UE::ConcertSharedSlate;
+
 	PrimarySortedColumn = ConcertBrowserUtils::IconColName;
 	PrimarySortMode = EColumnSortMode::Ascending;
 	SecondarySortedColumn = ConcertBrowserUtils::SessionColName;
 	SecondarySortMode = EColumnSortMode::Ascending;
 
-	return SAssignNew(SessionsView, SListView<TSharedPtr<FConcertSessionItem>>)
+	SessionsView = SNew(SListView<TSharedPtr<FConcertSessionItem>>)
 		.SelectionMode(ESelectionMode::Single)
 		.ListItemsSource(&Sessions)
 		.OnGenerateRow(this, &SConcertSessionBrowser::OnGenerateSessionRowWidget)
@@ -354,7 +360,15 @@ TSharedRef<SWidget> SConcertSessionBrowser::MakeSessionTableView()
 		.OnSelectionChanged(this, &SConcertSessionBrowser::OnSessionSelectionChanged)
 		.OnContextMenuOpening(this, &SConcertSessionBrowser::MakeContextualMenu)
 		.HeaderRow(
-			SNew(SHeaderRow)
+			SAssignNew(SessionHeaderRow, SHeaderRow)
+			.OnHiddenColumnsListChanged_Lambda([this, SaveCallback = InArgs._SaveColumnVisibilitySnapshot]()
+			{
+				if (SaveCallback.IsBound())
+				{
+					const FColumnVisibilitySnapshot Snapshot = SnapshotColumnVisibilityState(SessionHeaderRow.ToSharedRef());
+					SaveCallback.Execute(Snapshot);
+				}
+			})
 
 			+SHeaderRow::Column(ConcertBrowserUtils::IconColName)
 			.DefaultLabel(FText::GetEmpty())
@@ -362,30 +376,39 @@ TSharedRef<SWidget> SConcertSessionBrowser::MakeSessionTableView()
 			.SortMode(this, &SConcertSessionBrowser::GetColumnSortMode, ConcertBrowserUtils::IconColName)
 			.OnSort(this, &SConcertSessionBrowser::OnColumnSortModeChanged)
 			.FixedWidth(20)
+			.ShouldGenerateWidget(true)
 
 			+SHeaderRow::Column(ConcertBrowserUtils::SessionColName)
 			.DefaultLabel(LOCTEXT("SessioName", "Session"))
 			.SortPriority(this, &SConcertSessionBrowser::GetColumnSortPriority, ConcertBrowserUtils::SessionColName)
 			.SortMode(this, &SConcertSessionBrowser::GetColumnSortMode, ConcertBrowserUtils::SessionColName)
 			.OnSort(this, &SConcertSessionBrowser::OnColumnSortModeChanged)
+			.ShouldGenerateWidget(true)
 
 			+SHeaderRow::Column(ConcertBrowserUtils::ServerColName)
 			.DefaultLabel(LOCTEXT("Server", "Server"))
 			.SortPriority(this, &SConcertSessionBrowser::GetColumnSortPriority, ConcertBrowserUtils::ServerColName)
 			.SortMode(this, &SConcertSessionBrowser::GetColumnSortMode, ConcertBrowserUtils::ServerColName)
 			.OnSort(this, &SConcertSessionBrowser::OnColumnSortModeChanged)
+			.OnGetMenuContent_Lambda([this](){ return MakeHideColumnContextMenu(SessionHeaderRow.ToSharedRef(), ConcertBrowserUtils::ServerColName); })
 
 			+SHeaderRow::Column(ConcertBrowserUtils::ProjectColName)
 			.DefaultLabel(LOCTEXT("Project", "Project"))
 			.SortPriority(this, &SConcertSessionBrowser::GetColumnSortPriority, ConcertBrowserUtils::ProjectColName)
 			.SortMode(this, &SConcertSessionBrowser::GetColumnSortMode, ConcertBrowserUtils::ProjectColName)
 			.OnSort(this, &SConcertSessionBrowser::OnColumnSortModeChanged)
+			.OnGetMenuContent_Lambda([this](){ return MakeHideColumnContextMenu(SessionHeaderRow.ToSharedRef(), ConcertBrowserUtils::ProjectColName); })
 
 			+SHeaderRow::Column(ConcertBrowserUtils::VersionColName)
 			.DefaultLabel(LOCTEXT("Version", "Version"))
 			.SortPriority(this, &SConcertSessionBrowser::GetColumnSortPriority, ConcertBrowserUtils::VersionColName)
 			.SortMode(this, &SConcertSessionBrowser::GetColumnSortMode, ConcertBrowserUtils::VersionColName)
-			.OnSort(this, &SConcertSessionBrowser::OnColumnSortModeChanged));
+			.OnSort(this, &SConcertSessionBrowser::OnColumnSortModeChanged)
+			.OnGetMenuContent_Lambda([this](){ return MakeHideColumnContextMenu(SessionHeaderRow.ToSharedRef(), ConcertBrowserUtils::VersionColName); })
+			);
+
+	RestoreColumnVisibilityState(SessionHeaderRow.ToSharedRef(), InArgs._ColumnVisibilitySnapshot);
+	return SessionsView.ToSharedRef();
 }
 
 EColumnSortMode::Type SConcertSessionBrowser::GetColumnSortMode(const FName ColumnId) const
@@ -706,6 +729,8 @@ TSharedRef<SWidget> SConcertSessionBrowser::MakeSessionViewOptionsBar()
 			EUserInterfaceActionType::ToggleButton
 		);
 
+		MenuBuilder.AddSeparator();
+		UE::ConcertSharedSlate::AddEntriesForShowingHiddenRows(SessionHeaderRow.ToSharedRef(), MenuBuilder);
 		return MenuBuilder.MakeWidget();
 	};
 
diff --git a/Engine/Plugins/Developer/Concert/ConcertUI/ConcertSharedSlate/Source/ConcertSharedSlate/Private/Session/History/SSessionHistory.cpp b/Engine/Plugins/Developer/Concert/ConcertUI/ConcertSharedSlate/Source/ConcertSharedSlate/Private/Session/History/SSessionHistory.cpp
--- a/Engine/Plugins/Developer/Concert/ConcertUI/ConcertSharedSlate/Source/ConcertSharedSlate/Private/Session/History/SSessionHistory.cpp
+++ b/Engine/Plugins/Developer/Concert/ConcertUI/ConcertSharedSlate/Source/ConcertSharedSlate/Private/Session/History/SSessionHistory.cpp
@@ -67,7 +67,9 @@ void SSessionHistory::Construct(const FArguments& InArgs)
 		.PackageActivitiesVisibility(ActivityListViewOptions.Get(), &FConcertSessionActivitiesOptions::GetPackageActivitiesVisibility)
 		.TransactionActivitiesVisibility(ActivityListViewOptions.Get(), &FConcertSessionActivitiesOptions::GetTransactionActivitiesVisibility)
 		.DetailsAreaVisibility(EVisibility::Visible)
-		.IsAutoScrollEnabled(true);
+		.IsAutoScrollEnabled(true)
+		.ColumnVisibilitySnapshot(InArgs._ColumnVisibilitySnapshot)
+		.SaveColumnVisibilitySnapshot(InArgs._SaveColumnVisibilitySnapshot);
 
 	ChildSlot
 	[
@@ -101,7 +103,12 @@ void SSessionHistory::Construct(const FArguments& InArgs)
 		[
 			ActivityListViewOptions->MakeStatusBar(
 				TAttribute<int32>(ActivityListView.Get(), &SConcertSessionActivities::GetTotalActivityNum),
-				TAttribute<int32>(ActivityListView.Get(), &SConcertSessionActivities::GetDisplayedActivityNum))
+				TAttribute<int32>(ActivityListView.Get(), &SConcertSessionActivities::GetDisplayedActivityNum),
+				FConcertSessionActivitiesOptions::FExtendContextMenu::CreateLambda([this](FMenuBuilder& MenuBuilder)
+				{
+					MenuBuilder.AddSeparator();
+					AddEntriesForShowingHiddenRows(ActivityListView->GetHeaderRow().ToSharedRef(), MenuBuilder);
+				}))
 		]
 	];
 }
@@ -150,6 +157,11 @@ void SSessionHistory::HandleActivityAddedOrUpdated(const FConcertClientInfo& InC
 	}
 }
 
+void SSessionHistory::OnColumnVisibilitySettingsChanged(const FColumnVisibilitySnapshot& ColumnSnapshot)
+{
+	ActivityListView->OnColumnVisibilitySettingsChanged(ColumnSnapshot);
+}
+
 void SSessionHistory::OnSearchTextChanged(const FText& InSearchText)
 {
 	SearchedText = InSearchText;
diff --git a/Engine/Plugins/Developer/Concert/ConcertUI/ConcertSharedSlate/Source/ConcertSharedSlate/Public/ConcertHeaderRowUtils.h b/Engine/Plugins/Developer/Concert/ConcertUI/ConcertSharedSlate/Source/ConcertSharedSlate/Public/ConcertHeaderRowUtils.h
new file mode 100644
--- /dev/null
+++ b/Engine/Plugins/Developer/Concert/ConcertUI/ConcertSharedSlate/Source/ConcertSharedSlate/Public/ConcertHeaderRowUtils.h
@@ -0,0 +1,42 @@
+// Copyright Epic Games, Inc. All Rights Reserved.
+
+#pragma once
+
+#include "CoreMinimal.h"
+#include "ConcertHeaderRowUtils.generated.h"
+
+class FMenuBuilder;
+class SHeaderRow;
+
+USTRUCT()
+struct CONCERTSHAREDSLATE_API FColumnVisibilitySnapshot
+{
+	GENERATED_BODY()
+
+	UPROPERTY()
+	FString Snapshot;
+};
+
+namespace UE::ConcertSharedSlate
+{
+	DECLARE_DELEGATE_OneParam(FSaveColumnVisibilitySnapshot, const FColumnVisibilitySnapshot& /*Snapshot*/);
+	
+	/**
+	 * Creates a widget for hiding a SHeaderRow::FColumn when it is right-clicked.
+	 */
+	CONCERTSHAREDSLATE_API TSharedRef<SWidget> MakeHideColumnContextMenu(const TSharedRef<SHeaderRow>& HeaderRow, const FName ForColumnID);
+	CONCERTSHAREDSLATE_API void AddHideColumnEntry(const TSharedRef<SHeaderRow>& HeaderRow, const FName ForColumnID, FMenuBuilder& MenuBuilder);
+
+	/**
+	 * Inspects the hidden rows on the header row and an entry for showing each hidden column.
+	 * @param HeaderRow To retrieve column information
+	 * @param MenuBuilder Where to add the menu entries
+	 */
+	CONCERTSHAREDSLATE_API void AddEntriesForShowingHiddenRows(const TSharedRef<SHeaderRow>& HeaderRow, FMenuBuilder& MenuBuilder);
+
+	/** Exports the visibility state of each column so it can be restore later */
+	CONCERTSHAREDSLATE_API FColumnVisibilitySnapshot SnapshotColumnVisibilityState(const TSharedRef<SHeaderRow>& HeaderRow);
+	/** Restores the column visibilities from an exported state */
+	CONCERTSHAREDSLATE_API void RestoreColumnVisibilityState(const TSharedRef<SHeaderRow>& HeaderRow, const FColumnVisibilitySnapshot& Snapshot);
+};
+
diff --git a/Engine/Plugins/Developer/Concert/ConcertUI/ConcertSharedSlate/Source/ConcertSharedSlate/Public/Session/Activity/SConcertSessionActivities.h b/Engine/Plugins/Developer/Concert/ConcertUI/ConcertSharedSlate/Source/ConcertSharedSlate/Public/Session/Activity/SConcertSessionActivities.h
--- a/Engine/Plugins/Developer/Concert/ConcertUI/ConcertSharedSlate/Source/ConcertSharedSlate/Public/Session/Activity/SConcertSessionActivities.h
+++ b/Engine/Plugins/Developer/Concert/ConcertUI/ConcertSharedSlate/Source/ConcertSharedSlate/Public/Session/Activity/SConcertSessionActivities.h
@@ -5,12 +5,13 @@
 #include "CoreMinimal.h"
 #include "ActivityColumn.h"
 #include "Async/Future.h"
+#include "ConcertHeaderRowUtils.h"
+#include "ConcertSyncSessionTypes.h"
 #include "Misc/TextFilter.h"
 #include "Widgets/DeclarativeSyntaxSupport.h"
 #include "Widgets/SCompoundWidget.h"
 #include "Widgets/Views/SListView.h"
 #include "Widgets/Layout/SExpandableArea.h"
-#include "ConcertSyncSessionTypes.h"
 
 struct FConcertSessionActivity;
 struct FConcertClientInfo;
@@ -137,6 +138,12 @@ public:
 
 		/** Show/hide a message overlay above the activities list explaining why no activities are displayed. */
 		SLATE_ATTRIBUTE(FText, NoActivitiesReasonText)
+	
+		/** Optional snapshot to restore column visibilities with */
+		SLATE_ARGUMENT(FColumnVisibilitySnapshot, ColumnVisibilitySnapshot)
+		/** Called whenever the column visibility changes and should be saved */
+		SLATE_EVENT(UE::ConcertSharedSlate::FSaveColumnVisibilitySnapshot, SaveColumnVisibilitySnapshot)
+	
 	SLATE_END_ARGS();
 	
 	/**
@@ -177,10 +184,13 @@ public:
 
 	/** Asks the view to update the text search filter. */
 	FText UpdateTextFilter(const FText& InFilterText);
-
+	
+	void OnColumnVisibilitySettingsChanged(const FColumnVisibilitySnapshot& ColumnSnapshot);
+	
 	const TAttribute<FText>& GetHighlightText() const { return HighlightText; }
 	const TAttribute<ETimeFormat>& GetTimeFormat() const { return TimeFormat; }
 	TOptional<FConcertClientInfo> GetClientInfo(const FGuid& Guid) const { return GetActivityUserFn.IsBound() ? GetActivityUserFn.Execute(Guid) : TOptional<FConcertClientInfo>{}; }
+	TSharedPtr<SHeaderRow> GetHeaderRow() const { return HeaderRow; }
 
 private:
 
@@ -286,6 +296,9 @@ private:
 	TSharedPtr<TTextFilter<const FConcertSessionActivity&>> SearchTextFilter;
 
 	TSharedPtr<SHeaderRow> HeaderRow;
+
+	/** Flag used to suspend saving of the column visibility while OnColumnVisibilitySettingsChanged is active to prevent infinite event recursion. */
+	bool bIsUpdatingColumnVisibility = false;
 	
 	/** The expandable area under which the activity details are displayed. */
 	TSharedPtr<SExpandableArea> ExpandableDetails;
@@ -319,17 +332,20 @@ private:
 class CONCERTSHAREDSLATE_API FConcertSessionActivitiesOptions : public TSharedFromThis<FConcertSessionActivitiesOptions>
 {
 public:
-	/** Returns a Menu widgets containing the available options. */
-	TSharedRef<SWidget> MakeMenuWidget();
-
-	/** Makes a standard View Options widget, displaying the eye ball icon and showing the possible options. */
-	TSharedRef<SWidget> MakeViewOptionsWidget();
+	
+	DECLARE_DELEGATE_OneParam(FExtendContextMenu, FMenuBuilder&)
 
+	/** Makes a status bar widget displaying the activity shown and the standard view options button. */
+	TSharedRef<SWidget> MakeStatusBar(TAttribute<int32> Total, TAttribute<int32> Displayed, TOptional<FExtendContextMenu> ExtendMenu = {});
+	
 	/** Makes a widgets saying how many operations are shown with respect to those filtered out. */
 	TSharedRef<SWidget> MakeDisplayedActivityCountWidget(TAttribute<int32> Total, TAttribute<int32> Displayed);
-
-	/** Makes a status bar widget displaying the activity shown and the standard view options button. */
-	TSharedRef<SWidget> MakeStatusBar(TAttribute<int32> Total, TAttribute<int32> Displayed);
+	
+	/** Makes a standard View Options widget, displaying the eye ball icon and showing the possible options. */
+	TSharedRef<SWidget> MakeViewOptionsWidget(TOptional<FExtendContextMenu> ExtendMenu = {});
+	
+	/** Returns a Menu widgets containing the available options. */
+	TSharedRef<SWidget> MakeMenuWidget(TOptional<FExtendContextMenu> ExtendMenu = {});
 
 	/** Returns the time format option. */
 	SConcertSessionActivities::ETimeFormat GetTimeFormat() const { return bDisplayRelativeTime ? SConcertSessionActivities::ETimeFormat::Relative : SConcertSessionActivities::ETimeFormat::Absolute; }
diff --git a/Engine/Plugins/Developer/Concert/ConcertUI/ConcertSharedSlate/Source/ConcertSharedSlate/Public/Session/Browser/SConcertSessionBrowser.h b/Engine/Plugins/Developer/Concert/ConcertUI/ConcertSharedSlate/Source/ConcertSharedSlate/Public/Session/Browser/SConcertSessionBrowser.h
--- a/Engine/Plugins/Developer/Concert/ConcertUI/ConcertSharedSlate/Source/ConcertSharedSlate/Public/Session/Browser/SConcertSessionBrowser.h
+++ b/Engine/Plugins/Developer/Concert/ConcertUI/ConcertSharedSlate/Source/ConcertSharedSlate/Public/Session/Browser/SConcertSessionBrowser.h
@@ -3,6 +3,8 @@
 #pragma once
 
 #include "CoreMinimal.h"
+#include "ConcertHeaderRowUtils.h"
+
 #include "Misc/TextFilter.h"
 #include "SlateFwd.h"
 #include "Widgets/DeclarativeSyntaxSupport.h"
@@ -80,6 +82,11 @@ public:
 	SLATE_EVENT(FCanRemoveSession, CanDeleteArchivedSession)
 	/** Ask the user to confirm deleting - most obvious implementation is showing a dialog box */
 	SLATE_EVENT(FCanRemoveSession, CanDeleteActiveSession)
+
+	/** Optional snapshot to restore column visibilities with */
+	SLATE_ARGUMENT(FColumnVisibilitySnapshot, ColumnVisibilitySnapshot)
+	/** Called whenever the column visibility changes and should be saved */
+	SLATE_EVENT(UE::ConcertSharedSlate::FSaveColumnVisibilitySnapshot, SaveColumnVisibilitySnapshot)
 	
 	SLATE_END_ARGS();
 
@@ -119,7 +126,7 @@ private:
 	// Layout the sessions view and controls.
 	TSharedRef<SWidget> MakeControlBar(const FArguments& InArgs);
 	TSharedRef<SWidget> MakeButtonBar(const FArguments& InArgs);
-	TSharedRef<SWidget> MakeSessionTableView();
+	TSharedRef<SWidget> MakeSessionTableView(const FArguments& InArgs);
 	TSharedRef<SWidget> MakeSessionViewOptionsBar();
 	TSharedRef<ITableRow> OnGenerateSessionRowWidget(TSharedPtr<FConcertSessionItem> Item, const TSharedRef<STableViewBase>& OwnerTable);
 
@@ -206,6 +213,7 @@ private:
 
 	// The session list view.
 	TSharedPtr<SListView<TSharedPtr<FConcertSessionItem>>> SessionsView;
+	TSharedPtr<SHeaderRow> SessionHeaderRow;
 
 	// The item corresponding to a row used to create/archive/restore a session. There is only one at the time
 	TSharedPtr<FConcertSessionItem> EditableSessionRow;
diff --git a/Engine/Plugins/Developer/Concert/ConcertUI/ConcertSharedSlate/Source/ConcertSharedSlate/Public/Session/History/SSessionHistory.h b/Engine/Plugins/Developer/Concert/ConcertUI/ConcertSharedSlate/Source/ConcertSharedSlate/Public/Session/History/SSessionHistory.h
--- a/Engine/Plugins/Developer/Concert/ConcertUI/ConcertSharedSlate/Source/ConcertSharedSlate/Public/Session/History/SSessionHistory.h
+++ b/Engine/Plugins/Developer/Concert/ConcertUI/ConcertSharedSlate/Source/ConcertSharedSlate/Public/Session/History/SSessionHistory.h
@@ -23,6 +23,11 @@ public:
 		SLATE_EVENT(SConcertSessionActivities::FGetTransactionEvent, GetTransactionEvent)
 		/** If bound, invoked when generating a row to add an overlay to a column. */
 		SLATE_EVENT(SConcertSessionActivities::FMakeColumnOverlayWidgetFunc, OnMakeColumnOverlayWidget)
+	
+		/** Optional snapshot to restore column visibilities with */
+		SLATE_ARGUMENT(FColumnVisibilitySnapshot, ColumnVisibilitySnapshot)
+		/** Called whenever the column visibility changes and should be saved */
+		SLATE_EVENT(UE::ConcertSharedSlate::FSaveColumnVisibilitySnapshot, SaveColumnVisibilitySnapshot)
 	SLATE_END_ARGS()
 
 	/**
@@ -42,6 +47,8 @@ public:
 
 	bool IsLastColumn(FName ColumnId) const { return ActivityListView->IsLastColumn(ColumnId); }
 
+	void OnColumnVisibilitySettingsChanged(const FColumnVisibilitySnapshot& ColumnSnapshot);
+
 private:
 
 	/** Invoked when the text in the search box widget changes. */
