Add last modified column to session browser

diff --git a/Engine/Plugins/Developer/Concert/ConcertUI/ConcertSharedSlate/Source/ConcertSharedSlate/Private/ConcertHeaderRowUtils.cpp b/Engine/Plugins/Developer/Concert/ConcertUI/ConcertSharedSlate/Source/ConcertSharedSlate/Private/ConcertHeaderRowUtils.cpp
--- a/Engine/Plugins/Developer/Concert/ConcertUI/ConcertSharedSlate/Source/ConcertSharedSlate/Private/ConcertHeaderRowUtils.cpp
+++ b/Engine/Plugins/Developer/Concert/ConcertUI/ConcertSharedSlate/Source/ConcertSharedSlate/Private/ConcertHeaderRowUtils.cpp
@@ -3,7 +3,9 @@
 #include "ConcertHeaderRowUtils.h"
 
 #include "Algo/AllOf.h"
+#include "Algo/AnyOf.h"
 #include "Algo/Find.h"
+#include "Algo/ForEach.h"
 
 #include "Dom/JsonObject.h"
 
@@ -123,6 +125,49 @@ namespace UE::ConcertSharedSlate
 			HeaderRow->SetShowGeneratedColumn(*ColumnId, bIsVisible);
 		}
 	}
+
+	FColumnVisibilityTransaction::FColumnVisibilityTransaction(TWeakPtr<SHeaderRow> HeaderRow)
+		: HeaderRow(HeaderRow)
+	{}
+
+	FColumnVisibilityTransaction& FColumnVisibilityTransaction::SetHeaderRow(TWeakPtr<SHeaderRow> InHeaderRow)
+	{
+		HeaderRow = InHeaderRow;
+		return *this;
+	}
+
+	FColumnVisibilityTransaction& FColumnVisibilityTransaction::SaveVisibilityAndSet(FName ColumnId, bool bShouldBeVisible)
+	{
+		if (const TSharedPtr<SHeaderRow> Pinned = HeaderRow.Pin())
+		{
+			SavedColumnVisibilities.Add(ColumnId, Pinned->IsColumnVisible(ColumnId));
+			OverridenColumnVisibilities.Add(ColumnId, bShouldBeVisible);
+			Pinned->SetShowGeneratedColumn(ColumnId, bShouldBeVisible);
+		}
+		return *this;
+	}
+
+	void FColumnVisibilityTransaction::ResetToSavedVisibilities(bool bOnlyResetIfNotOverriden)
+	{
+		const TSharedPtr<SHeaderRow> Pinned = HeaderRow.Pin();
+		if (!Pinned)
+		{
+			return;
+		}
+		
+		const bool bSkipUpdate = bOnlyResetIfNotOverriden
+			&& Algo::AnyOf(OverridenColumnVisibilities, [this, &Pinned](const TPair<FName, bool>& SavedColumnEntry)
+		{
+			return Pinned->IsColumnVisible(SavedColumnEntry.Key) != SavedColumnEntry.Value;
+		});
+		if (!bSkipUpdate)
+		{
+			Algo::ForEach(SavedColumnVisibilities, [this, &Pinned](const TPair<FName, bool>& SavedColumnEntry)
+			{
+				Pinned->SetShowGeneratedColumn(SavedColumnEntry.Key, SavedColumnEntry.Value);
+			});
+		}
+	}
 }
 
 #undef LOCTEXT_NAMESPACE
\ No newline at end of file
diff --git a/Engine/Plugins/Developer/Concert/ConcertUI/ConcertSharedSlate/Source/ConcertSharedSlate/Private/Session/Activity/PredefinedActivityColumns.cpp b/Engine/Plugins/Developer/Concert/ConcertUI/ConcertSharedSlate/Source/ConcertSharedSlate/Private/Session/Activity/PredefinedActivityColumns.cpp
--- a/Engine/Plugins/Developer/Concert/ConcertUI/ConcertSharedSlate/Source/ConcertSharedSlate/Private/Session/Activity/PredefinedActivityColumns.cpp
+++ b/Engine/Plugins/Developer/Concert/ConcertUI/ConcertSharedSlate/Source/ConcertSharedSlate/Private/Session/Activity/PredefinedActivityColumns.cpp
@@ -19,9 +19,9 @@
 
 namespace UE::ConcertSharedSlate::Private
 {
-	static FText GetActivityDateTime(const FConcertSessionActivity& Activity, SConcertSessionActivities::ETimeFormat TimeFormat)
+	static FText GetActivityDateTime(const FConcertSessionActivity& Activity, ETimeFormat TimeFormat)
 	{
-		return TimeFormat == SConcertSessionActivities::ETimeFormat::Relative ? ConcertFrontendUtils::FormatRelativeTime(Activity.Activity.EventTime) : FText::AsDateTime(Activity.Activity.EventTime);
+		return TimeFormat == ETimeFormat::Relative ? ConcertFrontendUtils::FormatRelativeTime(Activity.Activity.EventTime) : FText::AsDateTime(Activity.Activity.EventTime);
 	}
 }
 
diff --git a/Engine/Plugins/Developer/Concert/ConcertUI/ConcertSharedSlate/Source/ConcertSharedSlate/Private/Session/Activity/SConcertSessionActivities.cpp b/Engine/Plugins/Developer/Concert/ConcertUI/ConcertSharedSlate/Source/ConcertSharedSlate/Private/Session/Activity/SConcertSessionActivities.cpp
--- a/Engine/Plugins/Developer/Concert/ConcertUI/ConcertSharedSlate/Source/ConcertSharedSlate/Private/Session/Activity/SConcertSessionActivities.cpp
+++ b/Engine/Plugins/Developer/Concert/ConcertUI/ConcertSharedSlate/Source/ConcertSharedSlate/Private/Session/Activity/SConcertSessionActivities.cpp
@@ -43,12 +43,7 @@ const FName ShowLockActivitiesCheckBoxId        = TEXT("ShowLockActivities");
 const FName ShowPackageActivitiesCheckBoxId     = TEXT("ShowPackageActivities");
 const FName ShowTransactionActivitiesCheckBoxId = TEXT("ShowTransactionActivities");
 const FName ShowIgnoredActivitiesCheckBoxId     = TEXT("ShowIgnoredActivities");
-
-FText GetActivityDateTime(const FConcertSessionActivity& Activity, SConcertSessionActivities::ETimeFormat TimeFormat)
-{
-	return TimeFormat == SConcertSessionActivities::ETimeFormat::Relative ? ConcertFrontendUtils::FormatRelativeTime(Activity.Activity.EventTime) : FText::AsDateTime(Activity.Activity.EventTime);
-}
-
+	
 FText GetSummary(const FConcertSessionActivity& Activity, const FText& ClientName, bool bAsRichText)
 {
 	if (const FConcertSyncActivitySummary* Summary = Activity.ActivitySummary.Cast<FConcertSyncActivitySummary>())
@@ -127,7 +122,7 @@ void SConcertSessionActivityRow::Construct(const FArguments& InArgs, TSharedRef<
 	ColumnGetter = InColumnGetter;
 	
 	OnMakeColumnOverlayWidget = InArgs._OnMakeColumnOverlayWidget;
-	AbsoluteDateTime = ConcertSessionActivityUtils::GetActivityDateTime(*InActivity, SConcertSessionActivities::ETimeFormat::Absolute); // Cache the absolute time. It doesn't changes.
+	AbsoluteDateTime = ConcertFrontendUtils::FormatTime(InActivity->Activity.EventTime, ETimeFormat::Absolute); // Cache the absolute time. It doesn't changes.
 	ClientName = ConcertSessionActivityUtils::GetClientName(InActivityClient);
 
 	// Construct base class
diff --git a/Engine/Plugins/Developer/Concert/ConcertUI/ConcertSharedSlate/Source/ConcertSharedSlate/Private/Session/Browser/SConcertSessionBrowser.cpp b/Engine/Plugins/Developer/Concert/ConcertUI/ConcertSharedSlate/Source/ConcertSharedSlate/Private/Session/Browser/SConcertSessionBrowser.cpp
--- a/Engine/Plugins/Developer/Concert/ConcertUI/ConcertSharedSlate/Source/ConcertSharedSlate/Private/Session/Browser/SConcertSessionBrowser.cpp
+++ b/Engine/Plugins/Developer/Concert/ConcertUI/ConcertSharedSlate/Source/ConcertSharedSlate/Private/Session/Browser/SConcertSessionBrowser.cpp
@@ -146,7 +146,10 @@ static FConcertSessionItem GenerateSessionItem(const SessionType& Session)
 		Session.ServerInfo.ServerName,
 		Session.SessionInfo.Settings.ProjectName,
 		Version.ToString(),
-		Session.ServerInfo.ServerFlags,{}};
+		Session.ServerInfo.ServerFlags,
+		Session.SessionInfo.LastModified,
+		{}
+	};
 }
 
 }
@@ -405,6 +408,14 @@ TSharedRef<SWidget> SConcertSessionBrowser::MakeSessionTableView(const FArgument
 			.SortMode(this, &SConcertSessionBrowser::GetColumnSortMode, ConcertBrowserUtils::VersionColName)
 			.OnSort(this, &SConcertSessionBrowser::OnColumnSortModeChanged)
 			.OnGetMenuContent_Lambda([this](){ return MakeHideColumnContextMenu(SessionHeaderRow.ToSharedRef(), ConcertBrowserUtils::VersionColName); })
+
+			+SHeaderRow::Column(ConcertBrowserUtils::LastModifiedColName)
+			.DefaultLabel(LOCTEXT("LastModified", "Last Modified"))
+			.DefaultTooltip(LOCTEXT("LastModifiedTooltip", "The local time the session's directory was last modified"))
+			.SortPriority(this, &SConcertSessionBrowser::GetColumnSortPriority, ConcertBrowserUtils::LastModifiedColName)
+			.SortMode(this, &SConcertSessionBrowser::GetColumnSortMode, ConcertBrowserUtils::LastModifiedColName)
+			.OnSort(this, &SConcertSessionBrowser::OnColumnSortModeChanged)
+			.OnGetMenuContent_Lambda([this](){ return MakeHideColumnContextMenu(SessionHeaderRow.ToSharedRef(), ConcertBrowserUtils::LastModifiedColName); })
 			);
 
 	RestoreColumnVisibilityState(SessionHeaderRow.ToSharedRef(), InArgs._ColumnVisibilitySnapshot);
@@ -485,10 +496,28 @@ void SConcertSessionBrowser::SortSessionList()
 		{
 			return SortMode == EColumnSortMode::Ascending ? Lhs->SessionName < Rhs->SessionName : Lhs->SessionName > Rhs->SessionName;
 		}
-		else
+		else if (ColName == ConcertBrowserUtils::ServerColName)
 		{
 			return SortMode == EColumnSortMode::Ascending ? Lhs->ServerName < Rhs->ServerName : Lhs->ServerName > Rhs->ServerName;
 		}
+		else if (ColName == ConcertBrowserUtils::ProjectColName)
+		{
+			return SortMode == EColumnSortMode::Ascending ? Lhs->ProjectName < Rhs->ProjectName : Lhs->ProjectName > Rhs->ProjectName;
+		}
+		else if (ColName == ConcertBrowserUtils::VersionColName)
+		{
+			return SortMode == EColumnSortMode::Ascending ? Lhs->ProjectVersion < Rhs->ProjectVersion : Lhs->ProjectVersion > Rhs->ProjectVersion;
+		}
+		else if (ColName == ConcertBrowserUtils::LastModifiedColName)
+		{
+			return SortMode == EColumnSortMode::Ascending ? Lhs->LastModified < Rhs->LastModified : Lhs->LastModified > Rhs->LastModified;
+		}
+		else
+		{
+			// Did you add a new column?
+			checkNoEntry();
+			return false;
+		}
 	};
 
 	Sessions.Sort([&](const TSharedPtr<FConcertSessionItem>& Lhs, const TSharedPtr<FConcertSessionItem>& Rhs)
@@ -693,6 +722,25 @@ TSharedRef<SWidget> SConcertSessionBrowser::MakeSessionViewOptionsBar()
 	{
 		FMenuBuilder MenuBuilder(/*bInShouldCloseWindowAfterMenuSelection=*/true, nullptr);
 
+		MenuBuilder.AddMenuEntry(
+			LOCTEXT("DisplayLastModifiedInRelativeTime", "Display Relative Time"),
+			TAttribute<FText>::CreateLambda([this]()
+			{
+				const bool bIsVisible = SessionHeaderRow->IsColumnVisible(ConcertBrowserUtils::LastModifiedColName);
+				return bIsVisible
+					? LOCTEXT("DisplayLastModifiedInRelativeTime.Tooltip.Visible", "Display the Last Modified column in relative time?")
+					: LOCTEXT("DisplayLastModifiedInRelativeTime.Tooltip.Hidden", "Disabled because the Last Modified column is hidden.");
+			}),
+			FSlateIcon(),
+			FUIAction(
+				FExecuteAction::CreateSP(this, &SConcertSessionBrowser::OnFilterMenuChecked, ConcertBrowserUtils::LastModifiedCheckBoxMenuName),
+				FCanExecuteAction::CreateLambda([this] { return SessionHeaderRow->IsColumnVisible(ConcertBrowserUtils::LastModifiedColName); }),
+				FIsActionChecked::CreateLambda([this] { return PersistentSettings->LastModifiedTimeFormat == ETimeFormat::Relative; })),
+			NAME_None,
+			EUserInterfaceActionType::ToggleButton
+		);
+		MenuBuilder.AddSeparator();
+		
 		MenuBuilder.AddMenuEntry(
 			LOCTEXT("ActiveSessions_Label", "Active Sessions"),
 			LOCTEXT("ActiveSessions_Tooltip", "Displays Active Sessions"),
@@ -815,7 +863,21 @@ TSharedRef<SWidget> SConcertSessionBrowser::MakeSessionViewOptionsBar()
 
 void SConcertSessionBrowser::OnFilterMenuChecked(const FName MenuName)
 {
-	if (MenuName == ConcertBrowserUtils::ActiveSessionsCheckBoxMenuName)
+	if (MenuName == ConcertBrowserUtils::LastModifiedCheckBoxMenuName)
+	{
+		switch (PersistentSettings->LastModifiedTimeFormat)
+		{
+		case ETimeFormat::Relative:
+			PersistentSettings->LastModifiedTimeFormat = ETimeFormat::Absolute;
+			break;
+		case ETimeFormat::Absolute: 
+			PersistentSettings->LastModifiedTimeFormat = ETimeFormat::Relative;
+			break;
+		default:
+			checkNoEntry();
+		}
+	}
+	else if (MenuName == ConcertBrowserUtils::ActiveSessionsCheckBoxMenuName)
 	{
 		PersistentSettings->bShowActiveSessions = !PersistentSettings->bShowActiveSessions;
 	}
diff --git a/Engine/Plugins/Developer/Concert/ConcertUI/ConcertSharedSlate/Source/ConcertSharedSlate/Private/Session/Browser/SNewSessionRow.cpp b/Engine/Plugins/Developer/Concert/ConcertUI/ConcertSharedSlate/Source/ConcertSharedSlate/Private/Session/Browser/SNewSessionRow.cpp
--- a/Engine/Plugins/Developer/Concert/ConcertUI/ConcertSharedSlate/Source/ConcertSharedSlate/Private/Session/Browser/SNewSessionRow.cpp
+++ b/Engine/Plugins/Developer/Concert/ConcertUI/ConcertSharedSlate/Source/ConcertSharedSlate/Private/Session/Browser/SNewSessionRow.cpp
@@ -7,6 +7,9 @@
 
 #include "EditorFontGlyphs.h"
 #include "EditorStyleSet.h"
+
+#include "Algo/ForEach.h"
+
 #include "Framework/Application/SlateApplication.h"
 #include "Styling/AppStyle.h"
 #include "Widgets/Input/SEditableTextBox.h"
@@ -33,11 +36,25 @@ void SNewSessionRow::Construct(const FArguments& InArgs, TSharedPtr<FConcertSess
 
 	// Construct base class
 	SMultiColumnTableRow<TSharedPtr<FConcertSessionItem>>::Construct(FSuperRowType::FArguments(), InOwnerTableView);
+	
+	// This rows are needed to correctly display the row
+	TemporaryColumnShower
+		.SetHeaderRow(InOwnerTableView->GetHeaderRow())
+		.SaveVisibilityAndSet(ConcertBrowserUtils::IconColName, true)
+		.SaveVisibilityAndSet(ConcertBrowserUtils::SessionColName, true)
+		.SaveVisibilityAndSet(ConcertBrowserUtils::ServerColName, true)
+		.SaveVisibilityAndSet(ConcertBrowserUtils::VersionColName, true)
+		.SaveVisibilityAndSet(ConcertBrowserUtils::ProjectColName, true);
 
 	// Fill the server combo.
 	UpdateServerList();
 }
 
+SNewSessionRow::~SNewSessionRow()
+{
+	TemporaryColumnShower.ResetToSavedVisibilities();
+}
+
 void SNewSessionRow::Tick(const FGeometry& AllottedGeometry, const double InCurrentTime, const float InDeltaTime)
 {
 	// Performance could be improved by only updating it when the server list has changed but it's no biggy...
@@ -129,9 +146,8 @@ TSharedRef<SWidget> SNewSessionRow::GenerateWidgetForColumn(const FName& ColumnN
 				]
 			];
 	}
-	else
+	else if (ColumnName == ConcertBrowserUtils::ProjectColName)
 	{
-		check(ColumnName == ConcertBrowserUtils::ProjectColName);
 		FText ProjectName = FText::AsCultureInvariant(FApp::GetProjectName());
 		return SNew(SBox)
 			.VAlign(VAlign_Center)
@@ -144,6 +160,11 @@ TSharedRef<SWidget> SNewSessionRow::GenerateWidgetForColumn(const FName& ColumnN
 				.OnKeyDownHandler(this, &SNewSessionRow::OnKeyDownHandler)
 			];
 	}
+	else
+	{
+		check(ColumnName == ConcertBrowserUtils::LastModifiedColName);
+		return SNullWidget::NullWidget;
+	}
 }
 
 TSharedRef<SWidget> SNewSessionRow::OnGenerateServersComboOptionWidget(TSharedPtr<FConcertServerInfo> ServerItem)
diff --git a/Engine/Plugins/Developer/Concert/ConcertUI/ConcertSharedSlate/Source/ConcertSharedSlate/Private/Session/Browser/SNewSessionRow.h b/Engine/Plugins/Developer/Concert/ConcertUI/ConcertSharedSlate/Source/ConcertSharedSlate/Private/Session/Browser/SNewSessionRow.h
--- a/Engine/Plugins/Developer/Concert/ConcertUI/ConcertSharedSlate/Source/ConcertSharedSlate/Private/Session/Browser/SNewSessionRow.h
+++ b/Engine/Plugins/Developer/Concert/ConcertUI/ConcertSharedSlate/Source/ConcertSharedSlate/Private/Session/Browser/SNewSessionRow.h
@@ -3,6 +3,7 @@
 #pragma once
 
 #include "CoreMinimal.h"
+#include "ConcertHeaderRowUtils.h"
 #include "ConcertMessageData.h"
 #include "Widgets/DeclarativeSyntaxSupport.h"
 #include "Widgets/Views/STableRow.h"
@@ -10,6 +11,7 @@
 
 class FConcertSessionItem;
 class SEditableTextBox;
+class SHeaderRow;
 class SWidget;
 
 /**
@@ -29,20 +31,18 @@ public:
 		, _OnAcceptFunc()
 		, _OnDeclineFunc()
 		, _HighlightText()
-	{
-	}
-
-	SLATE_ARGUMENT(FGetServersFunc, GetServerFunc)
-	SLATE_ARGUMENT(FAcceptFunc, OnAcceptFunc)
-	SLATE_ARGUMENT(FDeclineFunc, OnDeclineFunc)
-
-	SLATE_ATTRIBUTE(FText, HighlightText)
-	SLATE_ATTRIBUTE(FString, DefaultServerURL)
+	{}
+		SLATE_ARGUMENT(FGetServersFunc, GetServerFunc)
+		SLATE_ARGUMENT(FAcceptFunc, OnAcceptFunc)
+		SLATE_ARGUMENT(FDeclineFunc, OnDeclineFunc)
 
+		SLATE_ATTRIBUTE(FText, HighlightText)
+		SLATE_ATTRIBUTE(FString, DefaultServerURL)
 	SLATE_END_ARGS()
-
+	
 	void Construct(const FArguments& InArgs, TSharedPtr<FConcertSessionItem> InItem, const TSharedRef<STableViewBase>& InOwnerTableView);
-
+	virtual ~SNewSessionRow() override;
+	
 	virtual TSharedRef<SWidget> GenerateWidgetForColumn(const FName& ColumnName) override;
 	virtual void Tick(const FGeometry& AllottedGeometry, const double InCurrentTime, const float InDeltaTime) override;
 
@@ -65,6 +65,9 @@ private:
 
 	void UpdateServerList();
 
+	/** Helper for temporarily showing all columns required for this column */
+	UE::ConcertSharedSlate::FColumnVisibilityTransaction TemporaryColumnShower;
+	
 	/** Holds the new item to fill with session name and server. */
 	TWeakPtr<FConcertSessionItem> Item;
 	/** Servers displayed in the server combo box. */
diff --git a/Engine/Plugins/Developer/Concert/ConcertUI/ConcertSharedSlate/Source/ConcertSharedSlate/Private/Session/Browser/SSaveRestoreSessionRow.cpp b/Engine/Plugins/Developer/Concert/ConcertUI/ConcertSharedSlate/Source/ConcertSharedSlate/Private/Session/Browser/SSaveRestoreSessionRow.cpp
--- a/Engine/Plugins/Developer/Concert/ConcertUI/ConcertSharedSlate/Source/ConcertSharedSlate/Private/Session/Browser/SSaveRestoreSessionRow.cpp
+++ b/Engine/Plugins/Developer/Concert/ConcertUI/ConcertSharedSlate/Source/ConcertSharedSlate/Private/Session/Browser/SSaveRestoreSessionRow.cpp
@@ -6,6 +6,9 @@
 #include "Session/Browser/ConcertSessionItem.h"
 
 #include "EditorStyleSet.h"
+
+#include "Algo/ForEach.h"
+
 #include "Framework/Application/SlateApplication.h"
 #include "Internationalization/Regex.h"
 #include "Styling/AppStyle.h"
@@ -30,6 +33,19 @@ void SSaveRestoreSessionRow::Construct(const FArguments& InArgs, TSharedPtr<FCon
 
 	// Construct base class
 	SMultiColumnTableRow<TSharedPtr<FConcertSessionItem>>::Construct(FSuperRowType::FArguments(), InOwnerTableView);
+
+	// This rows are needed to correctly display the row
+	TemporaryColumnShower
+		.SetHeaderRow(InOwnerTableView->GetHeaderRow())
+		.SaveVisibilityAndSet(ConcertBrowserUtils::IconColName, true)
+		.SaveVisibilityAndSet(ConcertBrowserUtils::SessionColName, true)
+		.SaveVisibilityAndSet(ConcertBrowserUtils::ServerColName, true)
+		.SaveVisibilityAndSet(ConcertBrowserUtils::ProjectColName, true);
+}
+
+SSaveRestoreSessionRow::~SSaveRestoreSessionRow()
+{
+	TemporaryColumnShower.ResetToSavedVisibilities();
 }
 
 TBitArray<> SSaveRestoreSessionRow::GetWiresNeededByDepth() const
@@ -115,40 +131,43 @@ TSharedRef<SWidget> SSaveRestoreSessionRow::GenerateWidgetForColumn(const FName&
 				.Text(FText::FromString(ItemPin->ServerName))
 				.HighlightText(HighlightText)
 			]
-		]
-
-		+SHorizontalBox::Slot()
-		.AutoWidth()
-		.Padding(2.0f)
-		.HAlign(HAlign_Left)
-		[
-			SNew(SUniformGridPanel)
-			.SlotPadding(FMargin(1.0f, 0.0f))
-
-			// 'Accept' button
-			+SUniformGridPanel::Slot(0, 0)
-			[
-				ConcertBrowserUtils::MakePositiveActionButton(
-					FEditorStyle::GetBrush("Icons.Check"),
-					ItemPin->Type == FConcertSessionItem::EType::RestoreSession ? LOCTEXT("RestoreCheckIconTooltip", "Restore the session") : LOCTEXT("ArchiveCheckIconTooltip", "Archive the session"),
-					TAttribute<bool>::Create([this]() { return !EditableSessionName->GetText().IsEmpty(); }), // Enabled?
-					FOnClicked::CreateRaw(this, &SSaveRestoreSessionRow::OnAccept))
-			]
-
-			// 'Cancel' button
-			+SUniformGridPanel::Slot(1, 0)
-			[
-				ConcertBrowserUtils::MakeNegativeActionButton(
-					FEditorStyle::GetBrush("Icons.X"),
-					LOCTEXT("CancelTooltip", "Cancel"),
-					true, // Enabled?
-					FOnClicked::CreateRaw(this, &SSaveRestoreSessionRow::OnDecline))
-			]
 		];
 	}
+	else if (ColumnName == ConcertBrowserUtils::ProjectColName)
+	{
+		return SNew(SHorizontalBox)
+			+SHorizontalBox::Slot()
+			.AutoWidth()
+			.Padding(2.0f)
+			.HAlign(HAlign_Left)
+			[
+				SNew(SUniformGridPanel)
+				.SlotPadding(FMargin(1.0f, 0.0f))
+
+				// 'Accept' button
+				+SUniformGridPanel::Slot(0, 0)
+				[
+					ConcertBrowserUtils::MakePositiveActionButton(
+						FEditorStyle::GetBrush("Icons.Check"),
+						ItemPin->Type == FConcertSessionItem::EType::RestoreSession ? LOCTEXT("RestoreCheckIconTooltip", "Restore the session") : LOCTEXT("ArchiveCheckIconTooltip", "Archive the session"),
+						TAttribute<bool>::Create([this]() { return !EditableSessionName->GetText().IsEmpty(); }), // Enabled?
+						FOnClicked::CreateRaw(this, &SSaveRestoreSessionRow::OnAccept))
+				]
+
+				// 'Cancel' button
+				+SUniformGridPanel::Slot(1, 0)
+				[
+					ConcertBrowserUtils::MakeNegativeActionButton(
+						FEditorStyle::GetBrush("Icons.X"),
+						LOCTEXT("CancelTooltip", "Cancel"),
+						true, // Enabled?
+						FOnClicked::CreateRaw(this, &SSaveRestoreSessionRow::OnDecline))
+				]
+			];
+	}
 	else
 	{
-		check(ColumnName == ConcertBrowserUtils::ProjectColName || ColumnName == ConcertBrowserUtils::VersionColName);
+		check(ColumnName == ConcertBrowserUtils::VersionColName || ColumnName == ConcertBrowserUtils::LastModifiedColName);
 		return SNullWidget::NullWidget;
 	}
 }
diff --git a/Engine/Plugins/Developer/Concert/ConcertUI/ConcertSharedSlate/Source/ConcertSharedSlate/Private/Session/Browser/SSaveRestoreSessionRow.h b/Engine/Plugins/Developer/Concert/ConcertUI/ConcertSharedSlate/Source/ConcertSharedSlate/Private/Session/Browser/SSaveRestoreSessionRow.h
--- a/Engine/Plugins/Developer/Concert/ConcertUI/ConcertSharedSlate/Source/ConcertSharedSlate/Private/Session/Browser/SSaveRestoreSessionRow.h
+++ b/Engine/Plugins/Developer/Concert/ConcertUI/ConcertSharedSlate/Source/ConcertSharedSlate/Private/Session/Browser/SSaveRestoreSessionRow.h
@@ -3,6 +3,8 @@
 #pragma once
 
 #include "CoreMinimal.h"
+#include "ConcertHeaderRowUtils.h"
+
 #include "Misc/Attribute.h"
 #include "Session/Browser/ConcertSessionItem.h"
 #include "Widgets/Views/SListView.h"
@@ -28,14 +30,14 @@ public:
 		, _OnDeclineFunc()
 		, _HighlightText()
 	{}
-
-	SLATE_ARGUMENT(FAcceptFunc, OnAcceptFunc)
-	SLATE_ARGUMENT(FDeclineFunc, OnDeclineFunc)
-	SLATE_ATTRIBUTE(FText, HighlightText)
-
+		SLATE_ARGUMENT(FAcceptFunc, OnAcceptFunc)
+		SLATE_ARGUMENT(FDeclineFunc, OnDeclineFunc)
+		SLATE_ATTRIBUTE(FText, HighlightText)
 	SLATE_END_ARGS()
 	
 	void Construct(const FArguments& InArgs, TSharedPtr<FConcertSessionItem> Node, const TSharedRef<STableViewBase>& InOwnerTableView);
+	virtual ~SSaveRestoreSessionRow() override;
+	
 	virtual TSharedRef<SWidget> GenerateWidgetForColumn(const FName& ColumnName) override;
 	virtual void Tick(const FGeometry& AllottedGeometry, const double InCurrentTime, const float InDeltaTime) override;
 
@@ -56,6 +58,9 @@ public:
 
 private:
 	
+	/** Helper for temporarily showing all columns required for this column */
+	UE::ConcertSharedSlate::FColumnVisibilityTransaction TemporaryColumnShower;
+	
 	TWeakPtr<FConcertSessionItem> Item;
 	TSharedPtr<SEditableTextBox> EditableSessionName;
 	
diff --git a/Engine/Plugins/Developer/Concert/ConcertUI/ConcertSharedSlate/Source/ConcertSharedSlate/Private/Session/Browser/SSessionRow.cpp b/Engine/Plugins/Developer/Concert/ConcertUI/ConcertSharedSlate/Source/ConcertSharedSlate/Private/Session/Browser/SSessionRow.cpp
--- a/Engine/Plugins/Developer/Concert/ConcertUI/ConcertSharedSlate/Source/ConcertSharedSlate/Private/Session/Browser/SSessionRow.cpp
+++ b/Engine/Plugins/Developer/Concert/ConcertUI/ConcertSharedSlate/Source/ConcertSharedSlate/Private/Session/Browser/SSessionRow.cpp
@@ -2,11 +2,16 @@
 
 #include "SSessionRow.h"
 
+#include "ConcertFrontendUtils.h"
+
 #include "Session/Browser/ConcertBrowserUtils.h"
 #include "Session/Browser/ConcertSessionItem.h"
 
 #include "EditorFontGlyphs.h"
 #include "Framework/Application/SlateApplication.h"
+
+#include "Session/Browser/ConcertSessionBrowserSettings.h"
+
 #include "Styling/AppStyle.h"
 #include "Styling/CoreStyle.h"
 #include "Widgets/Layout/SBox.h"
@@ -141,7 +146,6 @@ TSharedRef<SWidget> SSessionRow::GenerateServerDefaultColumn(const FSlateFontInf
 
 TSharedRef<SWidget> SSessionRow::GenerateProjectColumn(const FSlateFontInfo &FontInfo, const FSlateColor &FontColor)
 {
-	//TSharedPtr<FConcertSessionItem> ItemPin = Item.Pin();
 	return SNew(SBox)
 		.VAlign(VAlign_Center)
 		[
@@ -151,7 +155,7 @@ TSharedRef<SWidget> SSessionRow::GenerateProjectColumn(const FSlateFontInfo &Fon
 			.IsReadOnly(true)
 			.Font(FontInfo)
 			.ColorAndOpacity(FontColor)
-			];
+		];
 }
 
 TSharedRef<SWidget> SSessionRow::GenerateVersionColumn(const FSlateFontInfo& FontInfo, const FSlateColor& FontColor)
@@ -165,7 +169,19 @@ TSharedRef<SWidget> SSessionRow::GenerateVersionColumn(const FSlateFontInfo& Fon
 			.IsReadOnly(true)
 			.Font(FontInfo)
 			.ColorAndOpacity(FontColor)
-			];
+		];
+}
+
+TSharedRef<SWidget> SSessionRow::GenerateLastModifiedColumn(const FSlateFontInfo& FontInfo, const FSlateColor& FontColor)
+{
+	return SNew(SBox)
+		.VAlign(VAlign_Center)
+		[
+			SNew(STextBlock)
+			.Text_Lambda([this](){ return ConcertFrontendUtils::FormatTime(Item.Pin()->LastModified, GetMutableDefault<UConcertSessionBrowserSettings>()->LastModifiedTimeFormat); })
+			.Font(FontInfo)
+			.ColorAndOpacity(FontColor)
+		];
 }
 
 
@@ -222,9 +238,13 @@ TSharedRef<SWidget> SSessionRow::GenerateWidgetForColumn(const FName& ColumnName
 		return GenerateProjectColumn(FontInfo, FontColor);
 	}
 
-	check(ColumnName == ConcertBrowserUtils::VersionColName);
-
-	return GenerateVersionColumn(FontInfo, FontColor);
+	if (ColumnName == ConcertBrowserUtils::VersionColName)
+	{
+		return GenerateVersionColumn(FontInfo, FontColor);
+	}
+	
+	check(ColumnName == ConcertBrowserUtils::LastModifiedColName);
+	return GenerateLastModifiedColumn(FontInfo, FontColor);
 }
 
 FReply SSessionRow::OnMouseButtonDoubleClick(const FGeometry& InMyGeometry, const FPointerEvent& InMouseEvent)
diff --git a/Engine/Plugins/Developer/Concert/ConcertUI/ConcertSharedSlate/Source/ConcertSharedSlate/Private/Session/Browser/SSessionRow.h b/Engine/Plugins/Developer/Concert/ConcertUI/ConcertSharedSlate/Source/ConcertSharedSlate/Private/Session/Browser/SSessionRow.h
--- a/Engine/Plugins/Developer/Concert/ConcertUI/ConcertSharedSlate/Source/ConcertSharedSlate/Private/Session/Browser/SSessionRow.h
+++ b/Engine/Plugins/Developer/Concert/ConcertUI/ConcertSharedSlate/Source/ConcertSharedSlate/Private/Session/Browser/SSessionRow.h
@@ -53,6 +53,7 @@ private:
 	TSharedRef<SWidget> GenerateServerDefaultColumn(const FSlateFontInfo& FontInfo, const FSlateColor& FontColor);
 	TSharedRef<SWidget> GenerateProjectColumn(const FSlateFontInfo& FontInfo, const FSlateColor& FontColor);
 	TSharedRef<SWidget> GenerateVersionColumn(const FSlateFontInfo& FontInfo, const FSlateColor& FontColor);
+	TSharedRef<SWidget> GenerateLastModifiedColumn(const FSlateFontInfo& FontInfo, const FSlateColor& FontColor);
 
 	void OnBeginEditingSessionName() { SessionNameText->EnterEditingMode(); }
 	bool OnValidatingSessionName(const FText& NewSessionName, FText& OutError);
diff --git a/Engine/Plugins/Developer/Concert/ConcertUI/ConcertSharedSlate/Source/ConcertSharedSlate/Public/ConcertFrontendUtils.h b/Engine/Plugins/Developer/Concert/ConcertUI/ConcertSharedSlate/Source/ConcertSharedSlate/Public/ConcertFrontendUtils.h
--- a/Engine/Plugins/Developer/Concert/ConcertUI/ConcertSharedSlate/Source/ConcertSharedSlate/Public/ConcertFrontendUtils.h
+++ b/Engine/Plugins/Developer/Concert/ConcertUI/ConcertSharedSlate/Source/ConcertSharedSlate/Public/ConcertFrontendUtils.h
@@ -7,9 +7,18 @@
 #include "Widgets/Layout/SBorder.h"
 #include "Widgets/Text/STextBlock.h"
 #include "Widgets/Layout/SExpandableArea.h"
+#include "ConcertFrontendUtils.generated.h"
 
 #define LOCTEXT_NAMESPACE "ConcertFrontendUtils"
 
+/** Defines how the time should be displayed in the date/time column. */
+UENUM()
+enum class ETimeFormat
+{
+	Relative, // Display relative time (23 seconds ago)
+	Absolute  // Display absolute time (April 7, 2019 - 10:33:52)
+};
+
 namespace ConcertFrontendUtils
 {
 	inline TSharedRef<SWidget> CreateDisplayName(const TAttribute<FText>& InDisplayName)
@@ -128,6 +137,13 @@ namespace ConcertFrontendUtils
 		}
 		return LOCTEXT("Now", "Now");
 	}
+
+	static FText FormatTime(const FDateTime& Time, ETimeFormat TimeFormat)
+	{
+		return TimeFormat == ETimeFormat::Relative
+			? FormatRelativeTime(Time)
+			: FText::AsDateTime(Time);
+	}
 };
 
 #undef LOCTEXT_NAMESPACE
diff --git a/Engine/Plugins/Developer/Concert/ConcertUI/ConcertSharedSlate/Source/ConcertSharedSlate/Public/ConcertHeaderRowUtils.h b/Engine/Plugins/Developer/Concert/ConcertUI/ConcertSharedSlate/Source/ConcertSharedSlate/Public/ConcertHeaderRowUtils.h
--- a/Engine/Plugins/Developer/Concert/ConcertUI/ConcertSharedSlate/Source/ConcertSharedSlate/Public/ConcertHeaderRowUtils.h
+++ b/Engine/Plugins/Developer/Concert/ConcertUI/ConcertSharedSlate/Source/ConcertSharedSlate/Public/ConcertHeaderRowUtils.h
@@ -39,5 +39,21 @@ namespace UE::ConcertSharedSlate
 	CONCERTSHAREDSLATE_API FColumnVisibilitySnapshot SnapshotColumnVisibilityState(const TSharedRef<SHeaderRow>& HeaderRow);
 	/** Restores the column visibilities from an exported state */
 	CONCERTSHAREDSLATE_API void RestoreColumnVisibilityState(const TSharedRef<SHeaderRow>& HeaderRow, const FColumnVisibilitySnapshot& Snapshot);
+
+	/** Helper for temporarily setting column visibilities to a certain setting. */
+	class FColumnVisibilityTransaction
+	{
+		TMap<FName, bool> SavedColumnVisibilities;
+		TMap<FName, bool> OverridenColumnVisibilities;
+		TWeakPtr<SHeaderRow> HeaderRow;
+	public:
+		
+		FColumnVisibilityTransaction(TWeakPtr<SHeaderRow> HeaderRow = nullptr);
+		FColumnVisibilityTransaction& SetHeaderRow(TWeakPtr<SHeaderRow> InHeaderRow);
+
+		FColumnVisibilityTransaction& SaveVisibilityAndSet(FName ColumnId, bool bShouldBeVisible);
+		/** @param bOnlyResetIfNotOverriden Only reset to the old visibilities if none has changed from the saved state (i.e. user changed it during the operation) */
+		void ResetToSavedVisibilities(bool bOnlyResetIfNotOverriden = true);
+	};
 };
 
diff --git a/Engine/Plugins/Developer/Concert/ConcertUI/ConcertSharedSlate/Source/ConcertSharedSlate/Public/Session/Activity/SConcertSessionActivities.h b/Engine/Plugins/Developer/Concert/ConcertUI/ConcertSharedSlate/Source/ConcertSharedSlate/Public/Session/Activity/SConcertSessionActivities.h
--- a/Engine/Plugins/Developer/Concert/ConcertUI/ConcertSharedSlate/Source/ConcertSharedSlate/Public/Session/Activity/SConcertSessionActivities.h
+++ b/Engine/Plugins/Developer/Concert/ConcertUI/ConcertSharedSlate/Source/ConcertSharedSlate/Public/Session/Activity/SConcertSessionActivities.h
@@ -4,6 +4,8 @@
 
 #include "CoreMinimal.h"
 #include "ActivityColumn.h"
+#include "ConcertFrontendUtils.h"
+
 #include "Async/Future.h"
 #include "ConcertHeaderRowUtils.h"
 #include "ConcertSyncSessionTypes.h"
@@ -39,14 +41,7 @@ ENUM_CLASS_FLAGS(EConcertActivityFilterFlags);
 class CONCERTSHAREDSLATE_API SConcertSessionActivities : public SCompoundWidget
 {
 public:
-
-	/** Defines how the time should be displayed in the date/time column. */
-	enum class ETimeFormat
-	{
-		Relative, // Display relative time (23 seconds ago)
-		Absolute  // Display absolute time (April 7, 2019 - 10:33:52)
-	};
-
+	
 	/** Used to pull activities from a session. Used to fetch and display the activities of an archived session. */
 	DECLARE_DELEGATE_RetVal_ThreeParams(bool, FFetchActivitiesFunc,
 		TArray<TSharedPtr<FConcertSessionActivity>>& /*InOutActivities*/,
@@ -348,7 +343,7 @@ public:
 	TSharedRef<SWidget> MakeMenuWidget(TOptional<FExtendContextMenu> ExtendMenu = {});
 
 	/** Returns the time format option. */
-	SConcertSessionActivities::ETimeFormat GetTimeFormat() const { return bDisplayRelativeTime ? SConcertSessionActivities::ETimeFormat::Relative : SConcertSessionActivities::ETimeFormat::Absolute; }
+	ETimeFormat GetTimeFormat() const { return bDisplayRelativeTime ? ETimeFormat::Relative : ETimeFormat::Absolute; }
 
 	/** Returns whether the connection activities are listed in SConcertSessionActivities. */
 	EVisibility GetConnectionActivitiesVisibility() const { return bDisplayConnectionActivities ? EVisibility::Visible : EVisibility::Hidden; }
diff --git a/Engine/Plugins/Developer/Concert/ConcertUI/ConcertSharedSlate/Source/ConcertSharedSlate/Public/Session/Browser/ConcertBrowserUtils.h b/Engine/Plugins/Developer/Concert/ConcertUI/ConcertSharedSlate/Source/ConcertSharedSlate/Public/Session/Browser/ConcertBrowserUtils.h
--- a/Engine/Plugins/Developer/Concert/ConcertUI/ConcertSharedSlate/Source/ConcertSharedSlate/Public/Session/Browser/ConcertBrowserUtils.h
+++ b/Engine/Plugins/Developer/Concert/ConcertUI/ConcertSharedSlate/Source/ConcertSharedSlate/Public/Session/Browser/ConcertBrowserUtils.h
@@ -28,8 +28,10 @@ namespace ConcertBrowserUtils
 	static const FName ServerColName(TEXT("Server"));
 	static const FName ProjectColName(TEXT("Project"));
 	static const FName VersionColName(TEXT("Version"));
+	static const FName LastModifiedColName(TEXT("LastModified"));
 
 	// Name of the filter box in the View option.
+	static const FName LastModifiedCheckBoxMenuName(TEXT("LastModified"));
 	static const FName ActiveSessionsCheckBoxMenuName(TEXT("ActiveSessions"));
 	static const FName ArchivedSessionsCheckBoxMenuName(TEXT("ArchivedSessions"));
 	static const FName DefaultServerCheckBoxMenuName(TEXT("DefaultServer"));
diff --git a/Engine/Plugins/Developer/Concert/ConcertUI/ConcertSharedSlate/Source/ConcertSharedSlate/Public/Session/Browser/ConcertSessionBrowserSettings.h b/Engine/Plugins/Developer/Concert/ConcertUI/ConcertSharedSlate/Source/ConcertSharedSlate/Public/Session/Browser/ConcertSessionBrowserSettings.h
--- a/Engine/Plugins/Developer/Concert/ConcertUI/ConcertSharedSlate/Source/ConcertSharedSlate/Public/Session/Browser/ConcertSessionBrowserSettings.h
+++ b/Engine/Plugins/Developer/Concert/ConcertUI/ConcertSharedSlate/Source/ConcertSharedSlate/Public/Session/Browser/ConcertSessionBrowserSettings.h
@@ -2,6 +2,8 @@
 
 #pragma once
 
+#include "ConcertFrontendUtils.h"
+
 #include "UObject/Object.h"
 #include "ConcertSessionBrowserSettings.generated.h"
 
@@ -14,6 +16,9 @@ class CONCERTSHAREDSLATE_API UConcertSessionBrowserSettings : public UObject
 public:
 	UConcertSessionBrowserSettings() {}
 
+	UPROPERTY(config, EditAnywhere, Category="Multi-User Session Browser")
+	ETimeFormat LastModifiedTimeFormat = ETimeFormat::Absolute;
+	
 	UPROPERTY(config, EditAnywhere, Category="Multi-User Session Browser")
 	bool bShowActiveSessions = true;
 
diff --git a/Engine/Plugins/Developer/Concert/ConcertUI/ConcertSharedSlate/Source/ConcertSharedSlate/Public/Session/Browser/ConcertSessionItem.h b/Engine/Plugins/Developer/Concert/ConcertUI/ConcertSharedSlate/Source/ConcertSharedSlate/Public/Session/Browser/ConcertSessionItem.h
--- a/Engine/Plugins/Developer/Concert/ConcertUI/ConcertSharedSlate/Source/ConcertSharedSlate/Public/Session/Browser/ConcertSessionItem.h
+++ b/Engine/Plugins/Developer/Concert/ConcertUI/ConcertSharedSlate/Source/ConcertSharedSlate/Public/Session/Browser/ConcertSessionItem.h
@@ -44,6 +44,8 @@ public:
 	FString ProjectName;
 	FString ProjectVersion;
 	EConcertServerFlags ServerFlags = EConcertServerFlags::None;
+	
+	FDateTime LastModified;
 
 	FOnBeginEditConcertSessionNameRequest OnBeginEditSessionNameRequest; // Emitted when user press 'F2' or select 'Rename' from context menu.
 };
