Fix duplicated keys for localization - UE - Virtual Production

diff --git a/Engine/Plugins/Developer/Concert/ConcertUI/ConcertSharedSlate/Source/ConcertSharedSlate/Private/Session/Browser/SConcertSessionBrowser.cpp b/Engine/Plugins/Developer/Concert/ConcertUI/ConcertSharedSlate/Source/ConcertSharedSlate/Private/Session/Browser/SConcertSessionBrowser.cpp
--- a/Engine/Plugins/Developer/Concert/ConcertUI/ConcertSharedSlate/Source/ConcertSharedSlate/Private/Session/Browser/SConcertSessionBrowser.cpp
+++ b/Engine/Plugins/Developer/Concert/ConcertUI/ConcertSharedSlate/Source/ConcertSharedSlate/Private/Session/Browser/SConcertSessionBrowser.cpp
@@ -395,7 +395,7 @@ TSharedRef<SWidget> SConcertSessionBrowser::MakeSessionTableView(const FArgument
 			.OnGetMenuContent_Lambda([this](){ return MakeHideColumnContextMenu(SessionHeaderRow.ToSharedRef(), ConcertBrowserUtils::ServerColName); })
 
 			+SHeaderRow::Column(ConcertBrowserUtils::ProjectColName)
-			.DefaultLabel(LOCTEXT("Project", "Project"))
+			.DefaultLabel(LOCTEXT("ProjectCol", "Project"))
 			.SortPriority(this, &SConcertSessionBrowser::GetColumnSortPriority, ConcertBrowserUtils::ProjectColName)
 			.SortMode(this, &SConcertSessionBrowser::GetColumnSortMode, ConcertBrowserUtils::ProjectColName)
 			.OnSort(this, &SConcertSessionBrowser::OnColumnSortModeChanged)
diff --git a/Engine/Plugins/Editor/DisplayClusterLaunch/Source/DisplayClusterLaunchEditor/Private/DisplayClusterLaunchEditorModule.cpp b/Engine/Plugins/Editor/DisplayClusterLaunch/Source/DisplayClusterLaunchEditor/Private/DisplayClusterLaunchEditorModule.cpp
--- a/Engine/Plugins/Editor/DisplayClusterLaunch/Source/DisplayClusterLaunchEditor/Private/DisplayClusterLaunchEditorModule.cpp
+++ b/Engine/Plugins/Editor/DisplayClusterLaunch/Source/DisplayClusterLaunchEditor/Private/DisplayClusterLaunchEditorModule.cpp
@@ -983,7 +983,7 @@ void FDisplayClusterLaunchEditorModule::AddConsoleVariablesEditorAssetsToToolbar
 	{
 		MenuBuilder.BeginSection("DisplayClusterLaunchCvars", LOCTEXT("DisplayClusterLaunchCvars", "Console Variables"));
 		{
-			const FText ConsoleVariablesAssetTooltip = LOCTEXT("SelectConsoleVariablesAssetFormat","Select Console Variables Asset");
+			const FText ConsoleVariablesAssetTooltip = LOCTEXT("SelectConsoleVariablesAsset","Select Console Variables Asset");
 			
 			MenuBuilder.AddSubMenu(
 				TAttribute<FText>::Create([this](){ return FText::FromName(SelectedConsoleVariablesAssetName); }),
