Adding EditorStyle as an Editor target dependency

diff --git a/Engine/Plugins/Developer/Concert/ConcertUI/ConcertSharedSlate/Source/ConcertSharedSlate/ConcertSharedSlate.Build.cs b/Engine/Plugins/Developer/Concert/ConcertUI/ConcertSharedSlate/Source/ConcertSharedSlate/ConcertSharedSlate.Build.cs
--- a/Engine/Plugins/Developer/Concert/ConcertUI/ConcertSharedSlate/Source/ConcertSharedSlate/ConcertSharedSlate.Build.cs
+++ b/Engine/Plugins/Developer/Concert/ConcertUI/ConcertSharedSlate/Source/ConcertSharedSlate/ConcertSharedSlate.Build.cs
@@ -30,6 +30,16 @@ namespace UnrealBuildTool.Rules
 					"ToolWidgets"
 				}
 			);
+
+			if(Target.bBuildEditor == true)
+            {
+				PrivateDependencyModuleNames.AddRange(
+				new string[]
+				{
+					"EditorStyle",
+				}
+			);
+			}
 		}
 	}
 }
