Updating plugin to use core styles and remove EditorFontGlyphs

diff --git a/Engine/Plugins/Developer/Concert/ConcertApp/MultiUserServer/Source/MultiUserServer/Private/MultiUserServerModule.cpp b/Engine/Plugins/Developer/Concert/ConcertApp/MultiUserServer/Source/MultiUserServer/Private/MultiUserServerModule.cpp
--- a/Engine/Plugins/Developer/Concert/ConcertApp/MultiUserServer/Source/MultiUserServer/Private/MultiUserServerModule.cpp
+++ b/Engine/Plugins/Developer/Concert/ConcertApp/MultiUserServer/Source/MultiUserServer/Private/MultiUserServerModule.cpp
@@ -36,7 +36,6 @@ namespace UE::MultiUserServer
 	{
 		if (ensureMsgf(!WindowController.IsValid(), TEXT("InitSlateForServer is designed to be called at most once to create UI to run alongside the Multi User server.")))
 		{
-			PreInitializeMultiUser();
 			InitArgs.PostInitServerLoop.AddRaw(this, &FConcertServerUIModule::InitializeSlateApplication);
 			InitArgs.TickPostGameThread.AddRaw(this, &FConcertServerUIModule::TickSlate);
 		}
@@ -44,7 +43,7 @@ namespace UE::MultiUserServer
 	
 	void FConcertServerUIModule::PreInitializeMultiUser()
     {
-    	FModuleManager::Get().LoadModuleChecked("EditorStyle");
+    	FModuleManager::Get().LoadModuleChecked("SlateCore");
     	FConcertServerStyle::Initialize();
     	FConcertFrontendStyle::Initialize();
 
@@ -58,6 +57,8 @@ namespace UE::MultiUserServer
     	const FText ApplicationTitle = LOCTEXT("AppTitle", "Unreal Multi User Server");
     	FGlobalTabmanager::Get()->SetApplicationTitle(ApplicationTitle);
 
+		PreInitializeMultiUser();
+
     	CommandExecutor = MakeUnique<FConcertConsoleCommandExecutor>();
     	IModularFeatures::Get().RegisterModularFeature(IConsoleCommandExecutor::ModularFeatureName(), CommandExecutor.Get());
 
diff --git a/Engine/Plugins/Developer/Concert/ConcertUI/ConcertSharedSlate/Source/ConcertSharedSlate/ConcertSharedSlate.Build.cs b/Engine/Plugins/Developer/Concert/ConcertUI/ConcertSharedSlate/Source/ConcertSharedSlate/ConcertSharedSlate.Build.cs
--- a/Engine/Plugins/Developer/Concert/ConcertUI/ConcertSharedSlate/Source/ConcertSharedSlate/ConcertSharedSlate.Build.cs
+++ b/Engine/Plugins/Developer/Concert/ConcertUI/ConcertSharedSlate/Source/ConcertSharedSlate/ConcertSharedSlate.Build.cs
@@ -31,15 +31,6 @@ namespace UnrealBuildTool.Rules
 				}
 			);
 
-			if(Target.bBuildEditor == true)
-            {
-				PrivateDependencyModuleNames.AddRange(
-				new string[]
-				{
-					"EditorStyle",
-				}
-			);
-			}
 		}
 	}
 }
diff --git a/Engine/Plugins/Developer/Concert/ConcertUI/ConcertSharedSlate/Source/ConcertSharedSlate/Private/ConcertFrontendStyle.cpp b/Engine/Plugins/Developer/Concert/ConcertUI/ConcertSharedSlate/Source/ConcertSharedSlate/Private/ConcertFrontendStyle.cpp
--- a/Engine/Plugins/Developer/Concert/ConcertUI/ConcertSharedSlate/Source/ConcertSharedSlate/Private/ConcertFrontendStyle.cpp
+++ b/Engine/Plugins/Developer/Concert/ConcertUI/ConcertSharedSlate/Source/ConcertSharedSlate/Private/ConcertFrontendStyle.cpp
@@ -6,6 +6,7 @@
 #include "Styling/AppStyle.h"
 #include "Interfaces/IPluginManager.h"
 #include "SlateOptMacros.h"
+#include "Styling/StarshipCoreStyle.h"
 
 #define IMAGE_PLUGIN_BRUSH( RelativePath, ... ) FSlateImageBrush( FConcertFrontendStyle::InContent( RelativePath, ".png" ), __VA_ARGS__ )
 #define IMAGE_BRUSH(RelativePath, ...) FSlateImageBrush(StyleSet->RootToContentDir(RelativePath, TEXT(".png")), __VA_ARGS__)
@@ -108,8 +109,8 @@ void FConcertFrontendStyle::Initialize()
 
 	// Activity Text
 	{
-		FTextBlockStyle BoldText = FAppStyle::Get().GetWidgetStyle<FTextBlockStyle>("RichTextBlock.Bold");
-		StyleSet->Set("ActivityText.Bold", FTextBlockStyle(BoldText));
+		FTextBlockStyle BoldText = FAppStyle::Get().GetWidgetStyle<FTextBlockStyle>("NormalText");
+		StyleSet->Set("ActivityText.Bold", FTextBlockStyle(BoldText).SetFont(FStyleFonts::Get().NormalBold));
 	}
 
 	// Colors
diff --git a/Engine/Plugins/Developer/Concert/ConcertUI/ConcertSharedSlate/Source/ConcertSharedSlate/Private/Session/Activity/SConcertSessionActivities.cpp b/Engine/Plugins/Developer/Concert/ConcertUI/ConcertSharedSlate/Source/ConcertSharedSlate/Private/Session/Activity/SConcertSessionActivities.cpp
--- a/Engine/Plugins/Developer/Concert/ConcertUI/ConcertSharedSlate/Source/ConcertSharedSlate/Private/Session/Activity/SConcertSessionActivities.cpp
+++ b/Engine/Plugins/Developer/Concert/ConcertUI/ConcertSharedSlate/Source/ConcertSharedSlate/Private/Session/Activity/SConcertSessionActivities.cpp
@@ -905,7 +905,7 @@ TSharedRef<SWidget> FConcertSessionActivitiesOptions::MakeViewOptionsWidget(TOpt
 		.AutoWidth()
 		.VAlign(VAlign_Center)
 		[
-			SNew(SImage).Image(FAppStyle::Get().GetBrush("GenericViewButton")) // The eye ball image.
+			SNew(SImage).Image(FAppStyle::Get().GetBrush("Icons.Visible")) // The eye ball image.
 		]
 
 		+SHorizontalBox::Slot()
diff --git a/Engine/Plugins/Developer/Concert/ConcertUI/ConcertSharedSlate/Source/ConcertSharedSlate/Private/Session/Browser/SConcertSessionBrowser.cpp b/Engine/Plugins/Developer/Concert/ConcertUI/ConcertSharedSlate/Source/ConcertSharedSlate/Private/Session/Browser/SConcertSessionBrowser.cpp
--- a/Engine/Plugins/Developer/Concert/ConcertUI/ConcertSharedSlate/Source/ConcertSharedSlate/Private/Session/Browser/SConcertSessionBrowser.cpp
+++ b/Engine/Plugins/Developer/Concert/ConcertUI/ConcertSharedSlate/Source/ConcertSharedSlate/Private/Session/Browser/SConcertSessionBrowser.cpp
@@ -30,6 +30,7 @@
 #include "Widgets/Text/STextBlock.h"
 #include "Widgets/Views/SListView.h"
 #include "Widgets/Views/STableRow.h"
+#include "Styling/StyleColors.h"
 
 #define LOCTEXT_NAMESPACE "SConcertBrowser"
 
@@ -832,8 +833,8 @@ TSharedRef<SWidget> SConcertSessionBrowser::MakeSessionViewOptionsBar()
 		.AutoWidth()
 		[
 			SNew(SComboButton)
-			.ComboButtonStyle(FAppStyle::Get(), "GenericFilters.ComboButtonStyle")
-			.ForegroundColor(FLinearColor::White)
+			.ComboButtonStyle(FAppStyle::Get(), "ComboButton")
+			.ForegroundColor(FStyleColors::Foreground)
 			.ContentPadding(0)
 			.OnGetMenuContent_Lambda(AddFilterMenu)
 			.HasDownArrow(true)
@@ -846,7 +847,7 @@ TSharedRef<SWidget> SConcertSessionBrowser::MakeSessionViewOptionsBar()
 				.AutoWidth()
 				.VAlign(VAlign_Center)
 				[
-					SNew(SImage).Image(FAppStyle::Get().GetBrush("GenericViewButton")) // The eye ball image.
+					SNew(SImage).Image(FAppStyle::Get().GetBrush("Icons.Visible")) // The eye ball image.
 				]
 
 				+SHorizontalBox::Slot()
diff --git a/Engine/Plugins/Developer/Concert/ConcertUI/ConcertSharedSlate/Source/ConcertSharedSlate/Private/Session/Browser/SNewSessionRow.cpp b/Engine/Plugins/Developer/Concert/ConcertUI/ConcertSharedSlate/Source/ConcertSharedSlate/Private/Session/Browser/SNewSessionRow.cpp
--- a/Engine/Plugins/Developer/Concert/ConcertUI/ConcertSharedSlate/Source/ConcertSharedSlate/Private/Session/Browser/SNewSessionRow.cpp
+++ b/Engine/Plugins/Developer/Concert/ConcertUI/ConcertSharedSlate/Source/ConcertSharedSlate/Private/Session/Browser/SNewSessionRow.cpp
@@ -3,9 +3,8 @@
 #include "SNewSessionRow.h"
 
 #include "Session/Browser/ConcertBrowserUtils.h"
-#include "Session/Browser/ConcertSessionItem.h"
+#include "Session/Browser/ConcertSessionItem.h"		
 
-#include "EditorFontGlyphs.h"
 #include "Styling/AppStyle.h"
 
 #include "Algo/ForEach.h"
@@ -78,9 +77,8 @@ TSharedRef<SWidget> SNewSessionRow::GenerateWidgetForColumn(const FName& ColumnN
 			.VAlign(VAlign_Center)
 			.HAlign(HAlign_Center)
 			[
-				SNew(STextBlock)
-					.Font(FAppStyle::Get().GetFontStyle(ConcertBrowserUtils::IconColumnFontName))
-					.Text(FEditorFontGlyphs::Plus_Circle)
+				SNew(SImage)
+				.Image(FAppStyle::GetBrush("Icons.PlusCircle"))
 			];
 	}
 	else if (ColumnName == ConcertBrowserUtils::SessionColName)
@@ -288,9 +286,8 @@ TSharedRef<SWidget> SNewSessionRow::MakeSelectedServerWidget()
 		.AutoWidth()
 		.Padding(2, 0, 0, 0)
 		[
-			SNew(STextBlock)
-			.Font(FAppStyle::Get().GetFontStyle("FontAwesome.9"))
-			.Text(this, &SNewSessionRow::GetSelectedServerIgnoreVersionText)
+			SNew(SImage)
+			.Image(this, &SNewSessionRow::GetSelectedServerIgnoreVersionImage)
 			.ToolTipText(this, &SNewSessionRow::GetSelectedServerIgnoreVersionTooltip)
 		];
 }
@@ -319,13 +316,13 @@ FText SNewSessionRow::GetServerDisplayName(const FString& ServerName) const
 	return FText::FromString(ServerName);
 }
 
-FText SNewSessionRow::GetSelectedServerIgnoreVersionText() const
+const FSlateBrush* SNewSessionRow::GetSelectedServerIgnoreVersionImage() const
 {
 	if (ServersComboBox->GetSelectedItem() && (ServersComboBox->GetSelectedItem()->ServerFlags & EConcertServerFlags::IgnoreSessionRequirement) != EConcertServerFlags::None)
 	{
-		return FEditorFontGlyphs::Exclamation_Triangle;
+		return FAppStyle::GetBrush("Icons.Warning");
 	}
-	return FText();
+	return FAppStyle::GetNoBrush();
 }
 
 FText SNewSessionRow::GetSelectedServerIgnoreVersionTooltip() const
diff --git a/Engine/Plugins/Developer/Concert/ConcertUI/ConcertSharedSlate/Source/ConcertSharedSlate/Private/Session/Browser/SNewSessionRow.h b/Engine/Plugins/Developer/Concert/ConcertUI/ConcertSharedSlate/Source/ConcertSharedSlate/Private/Session/Browser/SNewSessionRow.h
--- a/Engine/Plugins/Developer/Concert/ConcertUI/ConcertSharedSlate/Source/ConcertSharedSlate/Private/Session/Browser/SNewSessionRow.h
+++ b/Engine/Plugins/Developer/Concert/ConcertUI/ConcertSharedSlate/Source/ConcertSharedSlate/Private/Session/Browser/SNewSessionRow.h
@@ -8,6 +8,7 @@
 #include "Widgets/DeclarativeSyntaxSupport.h"
 #include "Widgets/Views/STableRow.h"
 #include "Widgets/Input/SComboBox.h"
+#include "Styling/SlateBrush.h"
 
 class FConcertSessionItem;
 class SEditableTextBox;
@@ -51,7 +52,7 @@ private:
 	TSharedRef<SWidget> OnGenerateServersComboOptionWidget(TSharedPtr<FConcertServerInfo> Item);
 	TSharedRef<SWidget> MakeSelectedServerWidget();
 	FText GetSelectedServerText() const;
-	FText GetSelectedServerIgnoreVersionText() const;
+	const FSlateBrush* GetSelectedServerIgnoreVersionImage() const;
 	FText GetSelectedServerIgnoreVersionTooltip() const;
 	FText GetServerDisplayName(const FString& ServerName) const;
 
diff --git a/Engine/Plugins/Developer/Concert/ConcertUI/ConcertSharedSlate/Source/ConcertSharedSlate/Private/Session/Browser/SSessionRow.cpp b/Engine/Plugins/Developer/Concert/ConcertUI/ConcertSharedSlate/Source/ConcertSharedSlate/Private/Session/Browser/SSessionRow.cpp
--- a/Engine/Plugins/Developer/Concert/ConcertUI/ConcertSharedSlate/Source/ConcertSharedSlate/Private/Session/Browser/SSessionRow.cpp
+++ b/Engine/Plugins/Developer/Concert/ConcertUI/ConcertSharedSlate/Source/ConcertSharedSlate/Private/Session/Browser/SSessionRow.cpp
@@ -7,7 +7,6 @@
 #include "Session/Browser/ConcertBrowserUtils.h"
 #include "Session/Browser/ConcertSessionItem.h"
 
-#include "EditorFontGlyphs.h"
 #include "Framework/Application/SlateApplication.h"
 
 #include "Session/Browser/ConcertSessionBrowserSettings.h"
@@ -192,16 +191,16 @@ TSharedRef<SWidget> SSessionRow::GenerateIconColumn()
 		LOCTEXT("ActiveIconTooltip", "Active session") : LOCTEXT("ArchivedIconTooltip", "Archived Session");
 	const FSlateColor InStyle = ItemPin->Type == FConcertSessionItem::EType::ActiveSession ?
 		FAppStyle::Get().GetWidgetStyle<FButtonStyle>("FlatButton.Success").Normal.TintColor : FSlateColor::UseSubduedForeground();
-	const FText InGlyph = ItemPin->Type == FConcertSessionItem::EType::ActiveSession ? FEditorFontGlyphs::Circle : FEditorFontGlyphs::Archive;
+	// @todo LB: change FolderClosed to archive
+	const FSlateBrush* InGlyph = ItemPin->Type == FConcertSessionItem::EType::ActiveSession ? FAppStyle::GetBrush("Icons.FilledCircle") : FAppStyle::GetBrush("Icons.FolderClosed");
 	return SNew(SBox)
 		.VAlign(VAlign_Center)
 		.HAlign(HAlign_Center)
 		.Padding(2)
 		.ToolTipText(Tooltip)
 		[
-			SNew(STextBlock)
-			.Font(FAppStyle::Get().GetFontStyle(ConcertBrowserUtils::IconColumnFontName))
-			.Text(InGlyph)
+			SNew(SImage)
+			.Image(InGlyph)
 			.ColorAndOpacity(InStyle)
 		];
 }
diff --git a/Engine/Plugins/Developer/Concert/ConcertUI/ConcertSharedSlate/Source/ConcertSharedSlate/Public/Session/Browser/ConcertBrowserUtils.h b/Engine/Plugins/Developer/Concert/ConcertUI/ConcertSharedSlate/Source/ConcertSharedSlate/Public/Session/Browser/ConcertBrowserUtils.h
--- a/Engine/Plugins/Developer/Concert/ConcertUI/ConcertSharedSlate/Source/ConcertSharedSlate/Public/Session/Browser/ConcertBrowserUtils.h
+++ b/Engine/Plugins/Developer/Concert/ConcertUI/ConcertSharedSlate/Source/ConcertSharedSlate/Public/Session/Browser/ConcertBrowserUtils.h
@@ -8,7 +8,6 @@
 #include "IConcertSessionBrowserController.h"
 #include "Session/Browser/ConcertSessionItem.h"
 
-#include "EditorFontGlyphs.h"
 #include "Internationalization/Regex.h"
 #include "Layout/Visibility.h"
 #include "Misc/Attribute.h"
@@ -96,9 +95,8 @@ namespace ConcertBrowserUtils
 			.BorderImage(FAppStyle::Get().GetBrush("NoBorder"))
 			.ColorAndOpacity(FAppStyle::Get().GetWidgetStyle<FButtonStyle>("FlatButton.Warning").Normal.TintColor.GetSpecifiedColor())
 			[
-				SNew(STextBlock)
-				.Font(FAppStyle::Get().GetFontStyle("FontAwesome.9"))
-				.Text(FEditorFontGlyphs::Exclamation_Triangle)
+				SNew(SImage)
+				.Image(FAppStyle::GetBrush("Icons.Warning"))
 				.ToolTipText(GetServerVersionIgnoredTooltip())
 				.Visibility((InServerFlags & EConcertServerFlags::IgnoreSessionRequirement) != EConcertServerFlags::None ? EVisibility::Visible : EVisibility::Collapsed)
 			];
diff --git a/Engine/Source/Editor/EditorStyle/Private/StarshipStyle.cpp b/Engine/Source/Editor/EditorStyle/Private/StarshipStyle.cpp
--- a/Engine/Source/Editor/EditorStyle/Private/StarshipStyle.cpp
+++ b/Engine/Source/Editor/EditorStyle/Private/StarshipStyle.cpp
@@ -7711,26 +7711,11 @@ void FStarshipEditorStyle::FStyle::SetupMyBlueprintStyles()
 
 void FStarshipEditorStyle::FStyle::SetupStatusBarStyle()
 {
-	const FButtonStyle StatusBarButton = FButtonStyle(FStarshipCoreStyle::GetCoreStyle().GetWidgetStyle<FButtonStyle>("NoBorder"))
-		.SetNormalForeground(FStyleColors::Foreground)
-		.SetNormalPadding(FMargin(2, 2, 2, 2))
-		.SetPressedPadding(FMargin(2, 3, 2, 1));
-
-	const FComboButtonStyle StatusBarComboButton = FComboButtonStyle(FStarshipCoreStyle::GetCoreStyle().GetWidgetStyle<FComboButtonStyle>("ComboButton"))
-		.SetDownArrowImage(CORE_IMAGE_BRUSH_SVG("Starship/CoreWidgets/ComboBox/corner-dropdown", FVector2D(7.0f, 7.0f)))
-		.SetButtonStyle(StatusBarButton)
-		.SetDownArrowPadding(FMargin(0.0f))
-		.SetDownArrowAlignment(EVerticalAlignment::VAlign_Bottom);
-		
-	Set("StatusBar.StatusBarButton", StatusBarButton);
-	Set("StatusBar.StatusBarComboButton", StatusBarComboButton);
-
 	Set("StatusBar.ContentBrowserUp", new CORE_IMAGE_BRUSH_SVG("Starship/Common/chevron-up", Icon16x16));
 	Set("StatusBar.ContentBrowserDown", new CORE_IMAGE_BRUSH_SVG("Starship/Common/chevron-down", Icon16x16));
 	Set("StatusBar.DrawerShadow", new BOX_BRUSH("Starship/StatusBar/drawer-shadow-bottom", FMargin(10/64.0f, 20/64.f, 10/64.f, 0), FLinearColor(0,0,0,1)));
 	Set("StatusBar.DrawerBackground", new FSlateRoundedBoxBrush(FStyleColors::Panel, 5.0f));
 
-	Set("StatusBar.Height", 32.0f);
 	Set("StatusBar.Background", new FSlateColorBrush(FStyleColors::Panel));
 	Set("StatusBar.ProgressOverlay", new FSlateRoundedBoxBrush(FStyleColors::Transparent, 2.0f, FStyleColors::Panel, 1.0f, FVector2D(20, 8)));
 
diff --git a/Engine/Source/Runtime/SlateCore/Private/Styling/StarshipCoreStyle.cpp b/Engine/Source/Runtime/SlateCore/Private/Styling/StarshipCoreStyle.cpp
--- a/Engine/Source/Runtime/SlateCore/Private/Styling/StarshipCoreStyle.cpp
+++ b/Engine/Source/Runtime/SlateCore/Private/Styling/StarshipCoreStyle.cpp
@@ -1635,6 +1635,21 @@ void FStarshipCoreStyle::SetupDockingStyles(TSharedRef<FStyle>& Style)
 	Style->Set("Docking.Sidebar.DrawerBackground", new FSlateColorBrush(FStyleColors::Panel));
 	Style->Set("Docking.Sidebar.Background", new FSlateColorBrush(FStyleColors::Background));
 
+	const FButtonStyle StatusBarButton = FButtonStyle(NoBorder)
+		.SetNormalForeground(FStyleColors::Foreground)
+		.SetNormalPadding(FMargin(2, 2, 2, 2))
+		.SetPressedPadding(FMargin(2, 3, 2, 1));
+
+	const FComboButtonStyle StatusBarComboButton = FComboButtonStyle(Style->GetWidgetStyle<FComboButtonStyle>("ComboButton"))
+		.SetDownArrowImage(IMAGE_BRUSH_SVG("Starship/CoreWidgets/ComboBox/corner-dropdown", FVector2D(7.0f, 7.0f)))
+		.SetButtonStyle(StatusBarButton)
+		.SetDownArrowPadding(FMargin(0.0f))
+		.SetDownArrowAlignment(EVerticalAlignment::VAlign_Bottom);
+
+	Style->Set("StatusBar.StatusBarButton", StatusBarButton);
+	Style->Set("StatusBar.StatusBarComboButton", StatusBarComboButton);
+	Style->Set("StatusBar.Height", 32.0f);
+
 	// Create the regular border brush with all corners rounded and variants with certain corners squared.
 	{
 		const float CornerRadius = 5.0f;
