Fix broken UI in details area of activities

diff --git a/Engine/Plugins/Developer/Concert/ConcertUI/ConcertSharedSlate/Source/ConcertSharedSlate/Private/ConcertFrontendStyle.cpp b/Engine/Plugins/Developer/Concert/ConcertUI/ConcertSharedSlate/Source/ConcertSharedSlate/Private/ConcertFrontendStyle.cpp
--- a/Engine/Plugins/Developer/Concert/ConcertUI/ConcertSharedSlate/Source/ConcertSharedSlate/Private/ConcertFrontendStyle.cpp
+++ b/Engine/Plugins/Developer/Concert/ConcertUI/ConcertSharedSlate/Source/ConcertSharedSlate/Private/ConcertFrontendStyle.cpp
@@ -10,6 +10,7 @@
 
 #define IMAGE_PLUGIN_BRUSH( RelativePath, ... ) FSlateImageBrush( FConcertFrontendStyle::InContent( RelativePath, ".png" ), __VA_ARGS__ )
 #define IMAGE_BRUSH(RelativePath, ...) FSlateImageBrush(StyleSet->RootToContentDir(RelativePath, TEXT(".png")), __VA_ARGS__)
+#define BOX_BRUSH(RelativePath, ... ) FSlateBoxBrush( StyleSet->RootToContentDir(RelativePath, TEXT(".png") ), __VA_ARGS__ )
 #define IMAGE_PLUGIN_BRUSH_SVG( RelativePath, ... ) FSlateVectorImageBrush( FConcertFrontendStyle::InContent(RelativePath, ".svg"), __VA_ARGS__)
 
 FString FConcertFrontendStyle::InContent(const FString& RelativePath, const ANSICHAR* Extension)
@@ -47,6 +48,12 @@ void FConcertFrontendStyle::Initialize()
 	// Use this to change the opacity. Ex: In UE4, the icon looked fine at 80%, in UE5, icons looks homogeneous with others at 100%.
 	const FLinearColor IconColorAndOpacity(FLinearColor(1.f, 1.f, 1.f, 1.f));
 
+	// Expandable area
+	StyleSet->Set("DetailsView.CollapsedCategory", new BOX_BRUSH( "Common/GroupBorder", FMargin(4.0f/16.0f) ) );
+	StyleSet->Set("DetailsView.CategoryTop", new BOX_BRUSH( "PropertyView/DetailCategoryTop", FMargin( 4/16.0f, 8.0f/16.0f, 4/16.0f, 4/16.0f ) ) );
+	StyleSet->Set("DetailsView.CollapsedCategory_Hovered", new BOX_BRUSH( "Common/GroupBorder", FMargin(4.0f/16.0f), FLinearColor(0.5f,0.5f,0.5f,1.0f)  ) );
+	StyleSet->Set("DetailsView.CategoryTop_Hovered", new BOX_BRUSH( "PropertyView/DetailCategoryTop", FMargin( 4/16.0f, 8.0f/16.0f, 4/16.0f, 4/16.0f ), FLinearColor(0.5f,0.5f,0.5f,1.0f) ) );
+
 	// 16x16
 	StyleSet->Set("Concert.Persist",         new IMAGE_PLUGIN_BRUSH("Icons/icon_ConcertPersist_16x", Icon16x16));
 	StyleSet->Set("Concert.LockBackground",  new IMAGE_PLUGIN_BRUSH("Icons/icon_ConcertLockBackground_16x", Icon16x16));
diff --git a/Engine/Plugins/Developer/Concert/ConcertUI/ConcertSharedSlate/Source/ConcertSharedSlate/Public/ConcertFrontendUtils.h b/Engine/Plugins/Developer/Concert/ConcertUI/ConcertSharedSlate/Source/ConcertSharedSlate/Public/ConcertFrontendUtils.h
--- a/Engine/Plugins/Developer/Concert/ConcertUI/ConcertSharedSlate/Source/ConcertSharedSlate/Public/ConcertFrontendUtils.h
+++ b/Engine/Plugins/Developer/Concert/ConcertUI/ConcertSharedSlate/Source/ConcertSharedSlate/Public/ConcertFrontendUtils.h
@@ -3,6 +3,7 @@
 #pragma once
 
 #include "CoreMinimal.h"
+#include "ConcertFrontendStyle.h"
 #include "Styling/AppStyle.h"
 #include "Widgets/Layout/SBorder.h"
 #include "Widgets/Text/STextBlock.h"
@@ -97,9 +98,9 @@ namespace ConcertFrontendUtils
 	{
 		if (Area.IsTitleHovered())
 		{
-			return Area.IsExpanded() ? FAppStyle::Get().GetBrush("DetailsView.CategoryTop_Hovered") : FAppStyle::Get().GetBrush("DetailsView.CollapsedCategory_Hovered");
+			return Area.IsExpanded() ? FConcertFrontendStyle::Get()->GetBrush("DetailsView.CategoryTop_Hovered") : FConcertFrontendStyle::Get()->GetBrush("DetailsView.CollapsedCategory_Hovered");
 		}
-		return Area.IsExpanded() ? FAppStyle::Get().GetBrush("DetailsView.CategoryTop") : FAppStyle::Get().GetBrush("DetailsView.CollapsedCategory");
+		return Area.IsExpanded() ? FConcertFrontendStyle::Get()->GetBrush("DetailsView.CategoryTop") : FConcertFrontendStyle::Get()->GetBrush("DetailsView.CollapsedCategory");
 	}
 
 	static FText FormatRelativeTime(const FDateTime& EventTime, const FDateTime* CurrTime = nullptr)
