Deprecated the FTransactionObjectEvent constructor that took separate ID elements, in favor of the one that takes a FTransactionObjectChange

#rnx

diff --git a/Engine/Plugins/Developer/Concert/ConcertSync/ConcertSyncClient/Source/ConcertSyncClient/Private/ConcertClientTransactionBridge.cpp b/Engine/Plugins/Developer/Concert/ConcertSync/ConcertSyncClient/Source/ConcertSyncClient/Private/ConcertClientTransactionBridge.cpp
--- a/Engine/Plugins/Developer/Concert/ConcertSync/ConcertSyncClient/Source/ConcertSyncClient/Private/ConcertClientTransactionBridge.cpp
+++ b/Engine/Plugins/Developer/Concert/ConcertSync/ConcertSyncClient/Source/ConcertSyncClient/Private/ConcertClientTransactionBridge.cpp
@@ -246,13 +246,8 @@ struct FEditorTransactionNotification
 				{
 					DeltaChange.ChangedProperties.Add(PropertyData.PropertyName);
 				}
-				TransactionObjectEvent = FTransactionObjectEvent(TransactionContext.TransactionId, TransactionContext.OperationId, ETransactionObjectEventType::UndoRedo, DeltaChange, InTransactionAnnotation
-					, InObjectUpdate.ObjectId.ObjectPackageName
-					, InTransactionObject->GetFName()
-					, *InTransactionObject->GetPathName()
-					, InObjectUpdate.ObjectId.ObjectOuterPathName
-					, InObjectUpdate.ObjectId.ObjectExternalPackageName
-					, FName(*InTransactionObject->GetClass()->GetPathName()));
+				TransactionObjectEvent = FTransactionObjectEvent(TransactionContext.TransactionId, TransactionContext.OperationId, ETransactionObjectEventType::UndoRedo, ETransactionObjectChangeCreatedBy::TransactionRecord,
+					FTransactionObjectChange{ InObjectUpdate.ObjectId.ToTransactionObjectId(), MoveTemp(DeltaChange) }, InTransactionAnnotation);
 			}
 			GUnrealEd->HandleObjectTransacted(InTransactionObject, TransactionObjectEvent);
 		}
diff --git a/Engine/Plugins/Developer/Concert/ConcertSync/ConcertSyncCore/Source/ConcertSyncCore/Public/ConcertTransactionEvents.h b/Engine/Plugins/Developer/Concert/ConcertSync/ConcertSyncCore/Source/ConcertSyncCore/Public/ConcertTransactionEvents.h
--- a/Engine/Plugins/Developer/Concert/ConcertSync/ConcertSyncCore/Source/ConcertSyncCore/Public/ConcertTransactionEvents.h
+++ b/Engine/Plugins/Developer/Concert/ConcertSync/ConcertSyncCore/Source/ConcertSyncCore/Public/ConcertTransactionEvents.h
@@ -31,13 +31,7 @@ struct FConcertObjectId
 	FConcertObjectId() = default;
 
 	explicit FConcertObjectId(const UObject* InObject)
-		: FConcertObjectId(
-			*InObject->GetClass()->GetPathName(), 
-			InObject->GetPackage()->GetFName(), 
-			InObject->GetFName(), 
-			InObject->GetOuter() ? FName(*InObject->GetOuter()->GetPathName()) : FName(), 
-			InObject->GetExternalPackage() ? InObject->GetExternalPackage()->GetFName() : FName(), 
-			InObject->GetFlags())
+		: FConcertObjectId(FTransactionObjectId(InObject), InObject->GetFlags())
 	{
 	}
 
@@ -56,6 +50,26 @@ struct FConcertObjectId
 	{
 	}
 
+	FTransactionObjectId ToTransactionObjectId() const
+	{
+		FNameBuilder ObjectPathName;
+		ObjectOuterPathName.AppendString(ObjectPathName);
+		{
+			int32 Unused = 0;
+			if (ObjectPathName.ToView().FindChar(TEXT('.'), Unused) && !ObjectPathName.ToView().FindChar(SUBOBJECT_DELIMITER_CHAR, Unused))
+			{
+				ObjectPathName.AppendChar(SUBOBJECT_DELIMITER_CHAR);
+			}
+			else
+			{
+				ObjectPathName.AppendChar(TEXT('.'));
+			}
+		}
+		ObjectName.AppendString(ObjectPathName);
+
+		return FTransactionObjectId(ObjectPackageName, ObjectName, ObjectPathName.ToString(), ObjectOuterPathName, ObjectExternalPackageName, ObjectClassPathName);
+	}
+
 	UPROPERTY()
 	FName ObjectClassPathName;
 
diff --git a/Engine/Plugins/Developer/Concert/ConcertUI/ConcertSharedSlate/Source/ConcertSharedSlate/Private/Session/Activity/SConcertSessionActivities.cpp b/Engine/Plugins/Developer/Concert/ConcertUI/ConcertSharedSlate/Source/ConcertSharedSlate/Private/Session/Activity/SConcertSessionActivities.cpp
--- a/Engine/Plugins/Developer/Concert/ConcertUI/ConcertSharedSlate/Source/ConcertSharedSlate/Private/Session/Activity/SConcertSessionActivities.cpp
+++ b/Engine/Plugins/Developer/Concert/ConcertUI/ConcertSharedSlate/Source/ConcertSharedSlate/Private/Session/Activity/SConcertSessionActivities.cpp
@@ -828,16 +828,11 @@ void SConcertSessionActivities::DisplayTransactionDetails(const FConcertSessionA
 		DeltaChange.bHasExternalPackageChange = ExportedObject.ObjectData.NewExternalPackageName != FName();
 		DeltaChange.bHasPendingKillChange = ExportedObject.ObjectData.bIsPendingKill;
 
-		FString ObjectPathName = ExportedObject.ObjectId.ObjectOuterPathName.ToString() + TEXT(".") + ExportedObject.ObjectId.ObjectName.ToString();
-		TSharedPtr<FTransactionObjectEvent> Event = MakeShared<FTransactionObjectEvent>(InTransaction.TransactionId, InTransaction.OperationId, ETransactionObjectEventType::Finalized, MoveTemp(DeltaChange), nullptr
-			, ExportedObject.ObjectId.ObjectPackageName
-			, ExportedObject.ObjectId.ObjectName
-			, FName(*MoveTemp(ObjectPathName))
-			, ExportedObject.ObjectId.ObjectOuterPathName
-			, ExportedObject.ObjectId.ObjectExternalPackageName
-			, ExportedObject.ObjectId.ObjectClassPathName);
-
-		TransactionDiff.DiffMap.Emplace(FName(*ObjectPathName), MoveTemp(Event));
+		FTransactionObjectId TransactionObjectId = ExportedObject.ObjectId.ToTransactionObjectId();
+		TSharedPtr<FTransactionObjectEvent> Event = MakeShared<FTransactionObjectEvent>(InTransaction.TransactionId, InTransaction.OperationId, ETransactionObjectEventType::Finalized, ETransactionObjectChangeCreatedBy::TransactionRecord, 
+			FTransactionObjectChange{ TransactionObjectId, MoveTemp(DeltaChange) }, nullptr);
+
+		TransactionDiff.DiffMap.Emplace(TransactionObjectId.ObjectPathName, MoveTemp(Event));
 	}
 
 	TransactionDetailsPanel->SetSelectedTransaction(MoveTemp(TransactionDiff));
diff --git a/Engine/Source/Runtime/CoreUObject/Public/Misc/TransactionObjectEvent.h b/Engine/Source/Runtime/CoreUObject/Public/Misc/TransactionObjectEvent.h
--- a/Engine/Source/Runtime/CoreUObject/Public/Misc/TransactionObjectEvent.h
+++ b/Engine/Source/Runtime/CoreUObject/Public/Misc/TransactionObjectEvent.h
@@ -190,7 +190,7 @@ public:
 		check(OperationId.IsValid());
 	}
 
-	// TODO: Deprecate this
+	UE_DEPRECATED(5.1, "Use the constructor that takes a FTransactionObjectChange.")
 	FTransactionObjectEvent(const FGuid& InTransactionId, const FGuid& InOperationId, const ETransactionObjectEventType InEventType, const FTransactionObjectDeltaChange& InDeltaChange, const TSharedPtr<ITransactionObjectAnnotation>& InAnnotation
 		, const FName InOriginalObjectPackageName, const FName InOriginalObjectName, const FName InOriginalObjectPathName, const FName InOriginalObjectOuterPathName, const FName InOriginalObjectExternalPackageName, const FName InOriginalObjectClassPathName)
 		: FTransactionObjectEvent(InTransactionId, InOperationId, InEventType, ETransactionObjectChangeCreatedBy::TransactionRecord, 
