Change positioning of search button in concert browser based on UX feedback

diff --git a/Engine/Plugins/Developer/Concert/ConcertApp/MultiUserClient/Source/MultiUserClient/Private/Widgets/SConcertBrowser.cpp b/Engine/Plugins/Developer/Concert/ConcertApp/MultiUserClient/Source/MultiUserClient/Private/Widgets/SConcertBrowser.cpp
--- a/Engine/Plugins/Developer/Concert/ConcertApp/MultiUserClient/Source/MultiUserClient/Private/Widgets/SConcertBrowser.cpp
+++ b/Engine/Plugins/Developer/Concert/ConcertApp/MultiUserClient/Source/MultiUserClient/Private/Widgets/SConcertBrowser.cpp
@@ -1173,7 +1173,7 @@ TSharedRef<SWidget> SConcertClientSessionBrowser::MakeBrowserContent(TSharedPtr<
 					.ExtendSessionTable(this, &SConcertClientSessionBrowser::MakeOverlayedTableView)
 					.ExtendControllButtons(this, &SConcertClientSessionBrowser::ExtendControlButtons)
 					.ExtendSessionContextMenu(this, &SConcertClientSessionBrowser::ExtendSessionContextMenu)
-					.RightOfSearchBar()
+					.RightOfControlButtons()
 					[
 						MakeUserAndSettings()
 					]
diff --git a/Engine/Plugins/Developer/Concert/ConcertUI/ConcertSharedSlate/Source/ConcertSharedSlate/Private/Session/Browser/SConcertSessionBrowser.cpp b/Engine/Plugins/Developer/Concert/ConcertUI/ConcertSharedSlate/Source/ConcertSharedSlate/Private/Session/Browser/SConcertSessionBrowser.cpp
--- a/Engine/Plugins/Developer/Concert/ConcertUI/ConcertSharedSlate/Source/ConcertSharedSlate/Private/Session/Browser/SConcertSessionBrowser.cpp
+++ b/Engine/Plugins/Developer/Concert/ConcertUI/ConcertSharedSlate/Source/ConcertSharedSlate/Private/Session/Browser/SConcertSessionBrowser.cpp
@@ -2,13 +2,8 @@
 
 #include "Session/Browser/SConcertSessionBrowser.h"
 
-#include "Algo/ForEach.h"
 #include "ConcertFrontendStyle.h"
 #include "ConcertHeaderRowUtils.h"
-
-#include "Algo/Find.h"
-#include "Algo/Transform.h"
-
 #include "Session/Browser/ConcertBrowserUtils.h"
 #include "Session/Browser/ConcertSessionBrowserSettings.h"
 #include "Session/Browser/ConcertSessionItem.h"
@@ -17,6 +12,8 @@
 #include "Session/Browser/SSaveRestoreSessionRow.h"
 #include "Session/Browser/SSessionRow.h"
 
+#include "Algo/ForEach.h"
+#include "Algo/Find.h"
 #include "Framework/MultiBox/MultiBoxBuilder.h"
 #include "Misc/TextFilter.h"
 #include "Styling/AppStyle.h"
@@ -31,6 +28,7 @@
 #include "Widgets/Views/SListView.h"
 #include "Widgets/Views/STableRow.h"
 #include "Styling/StyleColors.h"
+#include "Styling/ToolBarStyle.h"
 
 #define LOCTEXT_NAMESPACE "SConcertBrowser"
 
@@ -92,6 +90,26 @@ TSharedRef<SWidget> SConcertSessionBrowser::MakeBrowserContent(const FArguments&
 				MakeControlBar(InArgs)
 			]
 
+			+SVerticalBox::Slot()
+			.AutoHeight()
+			[
+				SNew(SSeparator)
+				.Thickness(2.f)
+				.SeparatorImage(&FCoreStyle::Get().GetWidgetStyle<FToolBarStyle>("ToolBar").SeparatorBrush)
+			]
+
+			// The search text.
+			+SVerticalBox::Slot()
+			.AutoHeight()
+			.Padding(0.f, 4.f)
+			[
+				SAssignNew(SearchBox, SSearchBox)
+				.HintText(LOCTEXT("SearchHint", "Search Session"))
+				.OnTextChanged(this, &SConcertSessionBrowser::OnSearchTextChanged)
+				.OnTextCommitted(this, &SConcertSessionBrowser::OnSearchTextCommitted)
+				.DelayChangeNotificationsWhileTyping(true)
+			]
+
 			// Session list.
 			+SVerticalBox::Slot()
 			.FillHeight(1.0f)
@@ -281,13 +299,8 @@ TSharedRef<SWidget> SConcertSessionBrowser::MakeControlBar(const FArguments& InA
 		// The search text.
 		+SHorizontalBox::Slot()
 		.FillWidth(1.0)
-		.Padding(4.0f, 5.0f, 8.0f, 5.0f)
 		[
-			SAssignNew(SearchBox, SSearchBox)
-			.HintText(LOCTEXT("SearchHint", "Search Session"))
-			.OnTextChanged(this, &SConcertSessionBrowser::OnSearchTextChanged)
-			.OnTextCommitted(this, &SConcertSessionBrowser::OnSearchTextCommitted)
-			.DelayChangeNotificationsWhileTyping(true)
+			SNew(SSpacer)
 		]
 
 		// Optional: everything to the right of the search bar, e.g. user name and settings combo button
@@ -295,7 +308,7 @@ TSharedRef<SWidget> SConcertSessionBrowser::MakeControlBar(const FArguments& InA
 		.VAlign(VAlign_Center)
 		.AutoWidth()
 		[
-			InArgs._RightOfSearchBar.Widget
+			InArgs._RightOfControlButtons.Widget
 		];
 }
 
diff --git a/Engine/Plugins/Developer/Concert/ConcertUI/ConcertSharedSlate/Source/ConcertSharedSlate/Public/Session/Browser/SConcertSessionBrowser.h b/Engine/Plugins/Developer/Concert/ConcertUI/ConcertSharedSlate/Source/ConcertSharedSlate/Public/Session/Browser/SConcertSessionBrowser.h
--- a/Engine/Plugins/Developer/Concert/ConcertUI/ConcertSharedSlate/Source/ConcertSharedSlate/Public/Session/Browser/SConcertSessionBrowser.h
+++ b/Engine/Plugins/Developer/Concert/ConcertUI/ConcertSharedSlate/Source/ConcertSharedSlate/Public/Session/Browser/SConcertSessionBrowser.h
@@ -67,8 +67,8 @@ public:
 	SLATE_EVENT(FExtenderDelegate, ExtendControllButtons)
 	/** Extends the menu when the user right-clicks a session */
 	SLATE_EVENT(FExtendSessionContextMenu, ExtendSessionContextMenu)
-	/** Custom slot placed to the right of the search bar */
-	SLATE_NAMED_SLOT(FArguments, RightOfSearchBar)
+	/** Custom slot placed to the right the control button in the top-most bar */
+	SLATE_NAMED_SLOT(FArguments, RightOfControlButtons)
 	
 	/** Called when a live or archived session is clicked */
 	SLATE_EVENT(FSessionDelegate, OnSessionClicked)
