Adjust context menu in session and activity browser

diff --git a/Engine/Plugins/Developer/Concert/ConcertApp/MultiUserServer/Source/MultiUserServer/Private/Widgets/Clients/Logging/SConcertTransportLog.cpp b/Engine/Plugins/Developer/Concert/ConcertApp/MultiUserServer/Source/MultiUserServer/Private/Widgets/Clients/Logging/SConcertTransportLog.cpp
--- a/Engine/Plugins/Developer/Concert/ConcertApp/MultiUserServer/Source/MultiUserServer/Private/Widgets/Clients/Logging/SConcertTransportLog.cpp
+++ b/Engine/Plugins/Developer/Concert/ConcertApp/MultiUserServer/Source/MultiUserServer/Private/Widgets/Clients/Logging/SConcertTransportLog.cpp
@@ -63,7 +63,7 @@ namespace UE::MultiUserServer
 								FFilterWidgetArgs()
 								.RightOfSearchBar
 								(
-									CreateViewOptionsButton()
+									ConcertFrontendUtils::CreateViewOptionsComboButton(FOnGetContent::CreateSP(this, &SConcertTransportLog::CreateViewOptionsMenu))
 								)
 							)
 							: SNullWidget::NullWidget
@@ -89,7 +89,7 @@ namespace UE::MultiUserServer
 		PagedLogList->OnPageViewChanged().AddSP(this, &SConcertTransportLog::OnPageViewChanged);
 		
 		UMultiUserServerColumnVisibilitySettings::GetSettings()->OnTransportLogColumnVisibility().AddSP(this, &SConcertTransportLog::OnColumnVisibilitySettingsChanged);
-		UE::ConcertSharedSlate::RestoreColumnVisibilityState(HeaderRow.ToSharedRef(), UMultiUserServerColumnVisibilitySettings::GetSettings()->GetTransportLogColumnVisibility());
+		ConcertSharedSlate::RestoreColumnVisibilityState(HeaderRow.ToSharedRef(), UMultiUserServerColumnVisibilitySettings::GetSettings()->GetTransportLogColumnVisibility());
 		
 		ConcertTransportEvents::OnConcertTransportLoggingEnabledChangedEvent().AddSP(this, &SConcertTransportLog::OnConcertLoggingEnabledChanged);
 		OnConcertLoggingEnabledChanged(ConcertTransportEvents::IsLoggingEnabled());
@@ -113,20 +113,6 @@ namespace UE::MultiUserServer
 		ScrollToLog(Index);
 	}
 
-	TSharedRef<SWidget> SConcertTransportLog::CreateViewOptionsButton()
-	{
-		return SNew(SComboButton)
-			.ComboButtonStyle(FAppStyle::Get(), "SimpleComboButtonWithIcon")
-			.OnGetMenuContent(this, &SConcertTransportLog::CreateViewOptionsMenu)
-			.HasDownArrow(false)
-			.ButtonContent()
-			[
-				SNew(SImage)
-				.ColorAndOpacity(FSlateColor::UseForeground())
-				.Image(FAppStyle::Get().GetBrush("Icons.Settings"))
-			];
-	}
-
 	TSharedRef<SWidget> SConcertTransportLog::CreateTableView()
 	{
 		CreateHeaderRow();
diff --git a/Engine/Plugins/Developer/Concert/ConcertApp/MultiUserServer/Source/MultiUserServer/Private/Widgets/Clients/Logging/SConcertTransportLog.h b/Engine/Plugins/Developer/Concert/ConcertApp/MultiUserServer/Source/MultiUserServer/Private/Widgets/Clients/Logging/SConcertTransportLog.h
--- a/Engine/Plugins/Developer/Concert/ConcertApp/MultiUserServer/Source/MultiUserServer/Private/Widgets/Clients/Logging/SConcertTransportLog.h
+++ b/Engine/Plugins/Developer/Concert/ConcertApp/MultiUserServer/Source/MultiUserServer/Private/Widgets/Clients/Logging/SConcertTransportLog.h
@@ -74,7 +74,6 @@ namespace UE::MultiUserServer
 		bool bIsUpdatingColumnVisibility = false;
 
 		// Table view creation
-		TSharedRef<SWidget> CreateViewOptionsButton();
 		TSharedRef<SWidget> CreateTableView();
 		TSharedRef<SHeaderRow> CreateHeaderRow();
 		TSharedRef<ITableRow> OnGenerateActivityRowWidget(TSharedPtr<FConcertLogEntry> Item, const TSharedRef<STableViewBase>& OwnerTable) const;
diff --git a/Engine/Plugins/Developer/Concert/ConcertApp/MultiUserServer/Source/MultiUserServer/Private/Widgets/Clients/PackageTransmission/SPackageTransmissionTable.cpp b/Engine/Plugins/Developer/Concert/ConcertApp/MultiUserServer/Source/MultiUserServer/Private/Widgets/Clients/PackageTransmission/SPackageTransmissionTable.cpp
--- a/Engine/Plugins/Developer/Concert/ConcertApp/MultiUserServer/Source/MultiUserServer/Private/Widgets/Clients/PackageTransmission/SPackageTransmissionTable.cpp
+++ b/Engine/Plugins/Developer/Concert/ConcertApp/MultiUserServer/Source/MultiUserServer/Private/Widgets/Clients/PackageTransmission/SPackageTransmissionTable.cpp
@@ -64,9 +64,8 @@ namespace UE::MultiUserServer
 
 	TSharedRef<SWidget> SPackageTransmissionTable::CreateViewOptionsButton()
 	{
-		return SNew(SComboButton)
-			.ComboButtonStyle(FAppStyle::Get(), "SimpleComboButtonWithIcon")
-			.OnGetMenuContent_Lambda([this]()
+		return ConcertFrontendUtils::CreateViewOptionsComboButton(
+			FOnGetContent::CreateLambda([this]()
 			{
 				FMenuBuilder MenuBuilder(true, nullptr);
 
@@ -107,14 +106,7 @@ namespace UE::MultiUserServer
 				);
 				
 				return MenuBuilder.MakeWidget();
-			})
-			.HasDownArrow(false)
-			.ButtonContent()
-			[
-				SNew(SImage)
-				.ColorAndOpacity(FSlateColor::UseForeground())
-				.Image(FAppStyle::Get().GetBrush("Icons.Settings"))
-			];
+			}));
 	}
 
 	TSharedRef<SWidget> SPackageTransmissionTable::CreateTableView()
diff --git a/Engine/Plugins/Developer/Concert/ConcertApp/MultiUserServer/Source/MultiUserServer/Private/Widgets/SessionTabs/Live/PackageViewer/SConcertSessionPackageViewer.cpp b/Engine/Plugins/Developer/Concert/ConcertApp/MultiUserServer/Source/MultiUserServer/Private/Widgets/SessionTabs/Live/PackageViewer/SConcertSessionPackageViewer.cpp
--- a/Engine/Plugins/Developer/Concert/ConcertApp/MultiUserServer/Source/MultiUserServer/Private/Widgets/SessionTabs/Live/PackageViewer/SConcertSessionPackageViewer.cpp
+++ b/Engine/Plugins/Developer/Concert/ConcertApp/MultiUserServer/Source/MultiUserServer/Private/Widgets/SessionTabs/Live/PackageViewer/SConcertSessionPackageViewer.cpp
@@ -53,11 +53,21 @@ void SConcertSessionPackageViewer::Construct(const FArguments& InArgs)
 		.AutoHeight()
 		.Padding(1, 1)
 		[
-			SAssignNew(SearchBox, SSearchBox)
-			.HintText(LOCTEXT("SearchHint", "Search..."))
-			.OnTextChanged(this, &SConcertSessionPackageViewer::OnSearchTextChanged)
-			.OnTextCommitted(this, &SConcertSessionPackageViewer::OnSearchTextCommitted)
-			.DelayChangeNotificationsWhileTyping(true)
+			SNew(SHorizontalBox)
+			+SHorizontalBox::Slot()
+			.FillWidth(1.f)
+			[
+				SAssignNew(SearchBox, SSearchBox)
+				.HintText(LOCTEXT("SearchHint", "Search..."))
+				.OnTextChanged(this, &SConcertSessionPackageViewer::OnSearchTextChanged)
+				.OnTextCommitted(this, &SConcertSessionPackageViewer::OnSearchTextCommitted)
+				.DelayChangeNotificationsWhileTyping(true)
+			]
+			+SHorizontalBox::Slot()
+			.AutoWidth()
+			[
+				ActivityListViewOptions->MakeViewOptionsComboButton()
+			]
 		]
 
 		+SVerticalBox::Slot()
@@ -70,12 +80,8 @@ void SConcertSessionPackageViewer::Construct(const FArguments& InArgs)
 		[
 			ActivityListViewOptions->MakeStatusBar(
 				TAttribute<int32>(ActivityListView.Get(), &SConcertSessionActivities::GetTotalActivityNum),
-				TAttribute<int32>(ActivityListView.Get(), &SConcertSessionActivities::GetDisplayedActivityNum),
-				FConcertSessionActivitiesOptions::FExtendContextMenu::CreateLambda([this](FMenuBuilder& MenuBuilder)
-				{
-					MenuBuilder.AddSeparator();
-					AddEntriesForShowingHiddenRows(ActivityListView->GetHeaderRow().ToSharedRef(), MenuBuilder);
-				}))
+				TAttribute<int32>(ActivityListView.Get(), &SConcertSessionActivities::GetDisplayedActivityNum)
+				)
 		]
 	];
 }
diff --git a/Engine/Plugins/Developer/Concert/ConcertUI/ConcertSharedSlate/Source/ConcertSharedSlate/Private/ConcertHeaderRowUtils.cpp b/Engine/Plugins/Developer/Concert/ConcertUI/ConcertSharedSlate/Source/ConcertSharedSlate/Private/ConcertHeaderRowUtils.cpp
--- a/Engine/Plugins/Developer/Concert/ConcertUI/ConcertSharedSlate/Source/ConcertSharedSlate/Private/ConcertHeaderRowUtils.cpp
+++ b/Engine/Plugins/Developer/Concert/ConcertUI/ConcertSharedSlate/Source/ConcertSharedSlate/Private/ConcertHeaderRowUtils.cpp
@@ -20,32 +20,6 @@
 
 namespace UE::ConcertSharedSlate
 {
-	TSharedRef<SWidget> MakeHideColumnContextMenu(const TSharedRef<SHeaderRow>& HeaderRow, const FName ForColumnID)
-	{
-		constexpr bool bInShouldCloseWindowAfterMenuSelection = true;
-		FMenuBuilder MenuBuilder(bInShouldCloseWindowAfterMenuSelection, nullptr);
-		AddHideColumnEntry(HeaderRow, ForColumnID, MenuBuilder);
-		return MenuBuilder.MakeWidget();
-	}
-
-	void AddHideColumnEntry(const TSharedRef<SHeaderRow>& HeaderRow, const FName ForColumnID, FMenuBuilder& MenuBuilder)
-	{
-		const SHeaderRow::FColumn* FoundColumn = Algo::FindByPredicate(HeaderRow->GetColumns(), [ForColumnID](const SHeaderRow::FColumn& Column){ return  Column.ColumnId == ForColumnID; });
-		check(FoundColumn);
-
-		TWeakPtr<SHeaderRow> WeakHeaderRow = HeaderRow;
-		MenuBuilder.AddMenuEntry(
-			LOCTEXT("HideColumn", "Hide column"),
-			FText::Format(LOCTEXT("HideColumn_Tooltip", "Hides the {0} column. You can unhide it using the eye-icon."), FoundColumn->DefaultText.Get()),
-			FSlateIcon(),
-			FUIAction(
-				FExecuteAction::CreateLambda([WeakHeaderRow, ForColumnID](){ WeakHeaderRow.Pin()->SetShowGeneratedColumn(ForColumnID, !WeakHeaderRow.Pin()->IsColumnVisible(ForColumnID)); }),
-				FCanExecuteAction::CreateLambda([] { return true; })),
-			NAME_None,
-			EUserInterfaceActionType::Button
-			);
-	}
-
 	TSharedRef<SWidget> MakeTableContextMenu(const TSharedRef<SHeaderRow>& HeaderRow, TMap<FName, bool> ColumnsVisibleByDefault, bool bDefaultVisibility)
 	{
 		FMenuBuilder MenuBuilder(true, nullptr);
diff --git a/Engine/Plugins/Developer/Concert/ConcertUI/ConcertSharedSlate/Source/ConcertSharedSlate/Private/SConcertSessionRecovery.cpp b/Engine/Plugins/Developer/Concert/ConcertUI/ConcertSharedSlate/Source/ConcertSharedSlate/Private/SConcertSessionRecovery.cpp
--- a/Engine/Plugins/Developer/Concert/ConcertUI/ConcertSharedSlate/Source/ConcertSharedSlate/Private/SConcertSessionRecovery.cpp
+++ b/Engine/Plugins/Developer/Concert/ConcertUI/ConcertSharedSlate/Source/ConcertSharedSlate/Private/SConcertSessionRecovery.cpp
@@ -79,11 +79,23 @@ void SConcertSessionRecovery::Construct(const FArguments& InArgs)
 					.AutoHeight()
 					.Padding(0, 1, 0, 2)
 					[
-						SAssignNew(SearchBox, SSearchBox)
-						.HintText(LOCTEXT("SearchHint", "Search..."))
-						.OnTextChanged(this, &SConcertSessionRecovery::OnSearchTextChanged)
-						.OnTextCommitted(this, &SConcertSessionRecovery::OnSearchTextCommitted)
-						.DelayChangeNotificationsWhileTyping(true)
+						SNew(SHorizontalBox)
+
+						+SHorizontalBox::Slot()
+						.FillWidth(1.f)
+						[
+							SAssignNew(SearchBox, SSearchBox)
+							.HintText(LOCTEXT("SearchHint", "Search..."))
+							.OnTextChanged(this, &SConcertSessionRecovery::OnSearchTextChanged)
+							.OnTextCommitted(this, &SConcertSessionRecovery::OnSearchTextCommitted)
+							.DelayChangeNotificationsWhileTyping(true)
+						]
+						
+						+SHorizontalBox::Slot()
+						.AutoWidth()
+						[
+							ActivityViewOptions->MakeViewOptionsComboButton()
+						]
 					]
 
 					// Activity List
@@ -106,7 +118,8 @@ void SConcertSessionRecovery::Construct(const FArguments& InArgs)
 					[
 						ActivityViewOptions->MakeStatusBar(
 							TAttribute<int32>(ActivityView.Get(), &SConcertSessionActivities::GetTotalActivityNum),
-							TAttribute<int32>(ActivityView.Get(), &SConcertSessionActivities::GetDisplayedActivityNum))
+							TAttribute<int32>(ActivityView.Get(), &SConcertSessionActivities::GetDisplayedActivityNum)
+							)
 					]
 				]
 			]
diff --git a/Engine/Plugins/Developer/Concert/ConcertUI/ConcertSharedSlate/Source/ConcertSharedSlate/Private/Session/Activity/SConcertSessionActivities.cpp b/Engine/Plugins/Developer/Concert/ConcertUI/ConcertSharedSlate/Source/ConcertSharedSlate/Private/Session/Activity/SConcertSessionActivities.cpp
--- a/Engine/Plugins/Developer/Concert/ConcertUI/ConcertSharedSlate/Source/ConcertSharedSlate/Private/Session/Activity/SConcertSessionActivities.cpp
+++ b/Engine/Plugins/Developer/Concert/ConcertUI/ConcertSharedSlate/Source/ConcertSharedSlate/Private/Session/Activity/SConcertSessionActivities.cpp
@@ -246,7 +246,8 @@ void SConcertSessionActivities::Construct(const FArguments& InArgs)
 	}
 	const TSharedRef<UE::UndoHistory::IReflectionDataProvider> ReflectionDataProvider =
 		ConcertReflectionDataProvider ? ConcertReflectionDataProvider->ToSharedRef() : UE::UndoHistory::CreateDefaultReflectionProvider();
-	
+
+	CreateHeaderRow(InArgs);
 	ActiveFilterFlags = QueryActiveActivityFilters();
 	ChildSlot
 	[
@@ -273,7 +274,8 @@ void SConcertSessionActivities::Construct(const FArguments& InArgs)
 					.AllowOverscroll(EAllowOverscroll::No)
 					.OnListViewScrolled(this, &SConcertSessionActivities::OnListViewScrolled)
 					.OnSelectionChanged(this, &SConcertSessionActivities::OnListViewSelectionChanged)
-					.HeaderRow(CreateHeaderRow(InArgs))
+					.HeaderRow(HeaderRow)
+					.OnContextMenuOpening_Lambda([this](){ return UE::ConcertSharedSlate::MakeTableContextMenu(HeaderRow.ToSharedRef(), {}, true); })
 				]
 			]
 
@@ -402,11 +404,6 @@ TSharedRef<SHeaderRow> SConcertSessionActivities::CreateHeaderRow(const FArgumen
 		
 		// SHeaderRow owns the columns and deletes them when destroyed
 		FActivityColumn* ManagedByHeaderRow = new FActivityColumn(MoveTemp(Column));
-		const bool bCanBeHidden = !ManagedByHeaderRow->ShouldGenerateWidget.IsSet(); 
-		if (bCanBeHidden)
-		{
-			ManagedByHeaderRow->OnGetMenuContent = FOnGetContent::CreateLambda([this, ColumnId = Column.ColumnId](){ return MakeHideColumnContextMenu(HeaderRow.ToSharedRef(), ColumnId); });
-		}
 		HeaderRow->AddColumn(*ManagedByHeaderRow);
 	}
 
@@ -868,34 +865,7 @@ void SConcertSessionActivities::SetDetailsPanelVisibility(const SWidget* Visible
 	LoadingDetailsPanel->SetVisibility(VisiblePanel == LoadingDetailsPanel.Get() ? EVisibility::Visible : EVisibility::Collapsed);
 }
 
-TSharedRef<SWidget> FConcertSessionActivitiesOptions::MakeStatusBar(TAttribute<int32> Total, TAttribute<int32> Displayed, TOptional<FExtendContextMenu> ExtendMenu)
-{
-	return SNew(SHorizontalBox)
-
-	// Operation count.
-	+SHorizontalBox::Slot()
-	.AutoWidth()
-	.VAlign(VAlign_Center)
-	[
-		MakeDisplayedActivityCountWidget(MoveTemp(Total), MoveTemp(Displayed))
-	]
-
-	// Gap Filler
-	+SHorizontalBox::Slot()
-	.FillWidth(1.0)
-	[
-		SNew(SSpacer)
-	]
-
-	// View option
-	+SHorizontalBox::Slot()
-	.AutoWidth()
-	[
-		MakeViewOptionsWidget(ExtendMenu)
-	];
-}
-
-TSharedRef<SWidget> FConcertSessionActivitiesOptions::MakeDisplayedActivityCountWidget(TAttribute<int32> Total, TAttribute<int32> Displayed)
+TSharedRef<SWidget> FConcertSessionActivitiesOptions::MakeStatusBar(TAttribute<int32> Total, TAttribute<int32> Displayed)
 {
 	return SNew(STextBlock)
 		.Text_Lambda([Total = MoveTemp(Total), Displayed = MoveTemp(Displayed)]()
@@ -904,37 +874,15 @@ TSharedRef<SWidget> FConcertSessionActivitiesOptions::MakeDisplayedActivityCount
 			{
 				return FText::Format(LOCTEXT("OperationCount", "{0} operations"), Total.Get());
 			}
-			else
-			{
-				return FText::Format(LOCTEXT("PartialOperationCount", "Showing {0} of {1} {1}|plural(one=operation,other=operations)"), Displayed.Get(), Total.Get());
-			}
+			return FText::Format(LOCTEXT("PartialOperationCount", "Showing {0} of {1} {1}|plural(one=operation,other=operations)"), Displayed.Get(), Total.Get());
 		});
 }
 
-TSharedRef<SWidget> FConcertSessionActivitiesOptions::MakeViewOptionsWidget(TOptional<FExtendContextMenu> ExtendMenu)
+TSharedRef<SWidget> FConcertSessionActivitiesOptions::MakeViewOptionsComboButton(TOptional<FExtendContextMenu> ExtendMenu)
 {
-	return SNew(SComboButton)
-	.OnGetMenuContent(this, &FConcertSessionActivitiesOptions::MakeMenuWidget, MoveTemp(ExtendMenu))
-	.HasDownArrow(true)
-	.ButtonContent()
-	[
-		SNew(SHorizontalBox)
-
-		+SHorizontalBox::Slot()
-		.AutoWidth()
-		.VAlign(VAlign_Center)
-		[
-			SNew(SImage).Image(FAppStyle::Get().GetBrush("Icons.Visible")) // The eye ball image.
-		]
-
-		+SHorizontalBox::Slot()
-		.AutoWidth()
-		.Padding(2, 0, 0, 0)
-		.VAlign(VAlign_Center)
-		[
-			SNew(STextBlock).Text(LOCTEXT("ViewOptions", "View Options"))
-		]
-	];
+	return ConcertFrontendUtils::CreateViewOptionsComboButton(
+		FOnGetContent::CreateRaw(this, &FConcertSessionActivitiesOptions::MakeMenuWidget, MoveTemp(ExtendMenu))
+		);
 }
 
 TSharedRef<SWidget> FConcertSessionActivitiesOptions::MakeMenuWidget(TOptional<FExtendContextMenu> ExtendMenu)
diff --git a/Engine/Plugins/Developer/Concert/ConcertUI/ConcertSharedSlate/Source/ConcertSharedSlate/Private/Session/Browser/SConcertSessionBrowser.cpp b/Engine/Plugins/Developer/Concert/ConcertUI/ConcertSharedSlate/Source/ConcertSharedSlate/Private/Session/Browser/SConcertSessionBrowser.cpp
--- a/Engine/Plugins/Developer/Concert/ConcertUI/ConcertSharedSlate/Source/ConcertSharedSlate/Private/Session/Browser/SConcertSessionBrowser.cpp
+++ b/Engine/Plugins/Developer/Concert/ConcertUI/ConcertSharedSlate/Source/ConcertSharedSlate/Private/Session/Browser/SConcertSessionBrowser.cpp
@@ -101,13 +101,24 @@ TSharedRef<SWidget> SConcertSessionBrowser::MakeBrowserContent(const FArguments&
 			// The search text.
 			+SVerticalBox::Slot()
 			.AutoHeight()
-			.Padding(0.f, 4.f)
+			.Padding(0.f, 4.f, 0.f, 4.f)
 			[
-				SAssignNew(SearchBox, SSearchBox)
-				.HintText(LOCTEXT("SearchHint", "Search Session"))
-				.OnTextChanged(this, &SConcertSessionBrowser::OnSearchTextChanged)
-				.OnTextCommitted(this, &SConcertSessionBrowser::OnSearchTextCommitted)
-				.DelayChangeNotificationsWhileTyping(true)
+				SNew(SHorizontalBox)
+				+SHorizontalBox::Slot()
+				.FillWidth(1.f)
+				[
+					SAssignNew(SearchBox, SSearchBox)
+					.HintText(LOCTEXT("SearchHint", "Search Session"))
+					.OnTextChanged(this, &SConcertSessionBrowser::OnSearchTextChanged)
+					.OnTextCommitted(this, &SConcertSessionBrowser::OnSearchTextCommitted)
+					.DelayChangeNotificationsWhileTyping(true)
+				]
+
+				+SHorizontalBox::Slot()
+				.AutoWidth()
+				[
+					ConcertFrontendUtils::CreateViewOptionsComboButton(FOnGetContent::CreateSP(this, &SConcertSessionBrowser::MakeViewOptionsComboButtonMenu))
+				]
 			]
 
 			// Session list.
@@ -130,7 +141,7 @@ TSharedRef<SWidget> SConcertSessionBrowser::MakeBrowserContent(const FArguments&
 			.AutoHeight()
 			.Padding(FMargin(2.0f, 0.0f))
 			[
-				MakeSessionViewOptionsBar()
+				MakeSessionTableFooter()
 			]
 		];
 }
@@ -374,7 +385,7 @@ TSharedRef<SWidget> SConcertSessionBrowser::MakeSessionTableView(const FArgument
 		.OnGenerateRow(this, &SConcertSessionBrowser::OnGenerateSessionRowWidget)
 		.SelectionMode(ESelectionMode::Multi)
 		.OnSelectionChanged(this, &SConcertSessionBrowser::OnSessionSelectionChanged)
-		.OnContextMenuOpening(this, &SConcertSessionBrowser::MakeContextualMenu)
+		.OnContextMenuOpening(this, &SConcertSessionBrowser::MakeTableContextualMenu)
 		.HeaderRow(
 			SAssignNew(SessionHeaderRow, SHeaderRow)
 			.OnHiddenColumnsListChanged_Lambda([this, SaveCallback = InArgs._SaveColumnVisibilitySnapshot]()
@@ -406,29 +417,25 @@ TSharedRef<SWidget> SConcertSessionBrowser::MakeSessionTableView(const FArgument
 			.SortPriority(this, &SConcertSessionBrowser::GetColumnSortPriority, ConcertBrowserUtils::ServerColName)
 			.SortMode(this, &SConcertSessionBrowser::GetColumnSortMode, ConcertBrowserUtils::ServerColName)
 			.OnSort(this, &SConcertSessionBrowser::OnColumnSortModeChanged)
-			.OnGetMenuContent_Lambda([this](){ return MakeHideColumnContextMenu(SessionHeaderRow.ToSharedRef(), ConcertBrowserUtils::ServerColName); })
-
+			
 			+SHeaderRow::Column(ConcertBrowserUtils::ProjectColName)
 			.DefaultLabel(LOCTEXT("ProjectCol", "Project"))
 			.SortPriority(this, &SConcertSessionBrowser::GetColumnSortPriority, ConcertBrowserUtils::ProjectColName)
 			.SortMode(this, &SConcertSessionBrowser::GetColumnSortMode, ConcertBrowserUtils::ProjectColName)
 			.OnSort(this, &SConcertSessionBrowser::OnColumnSortModeChanged)
-			.OnGetMenuContent_Lambda([this](){ return MakeHideColumnContextMenu(SessionHeaderRow.ToSharedRef(), ConcertBrowserUtils::ProjectColName); })
-
+			
 			+SHeaderRow::Column(ConcertBrowserUtils::VersionColName)
 			.DefaultLabel(LOCTEXT("Version", "Version"))
 			.SortPriority(this, &SConcertSessionBrowser::GetColumnSortPriority, ConcertBrowserUtils::VersionColName)
 			.SortMode(this, &SConcertSessionBrowser::GetColumnSortMode, ConcertBrowserUtils::VersionColName)
 			.OnSort(this, &SConcertSessionBrowser::OnColumnSortModeChanged)
-			.OnGetMenuContent_Lambda([this](){ return MakeHideColumnContextMenu(SessionHeaderRow.ToSharedRef(), ConcertBrowserUtils::VersionColName); })
-
+			
 			+SHeaderRow::Column(ConcertBrowserUtils::LastModifiedColName)
 			.DefaultLabel(LOCTEXT("LastModified", "Last Modified"))
 			.DefaultTooltip(LOCTEXT("LastModifiedTooltip", "The local time the session's directory was last modified"))
 			.SortPriority(this, &SConcertSessionBrowser::GetColumnSortPriority, ConcertBrowserUtils::LastModifiedColName)
 			.SortMode(this, &SConcertSessionBrowser::GetColumnSortMode, ConcertBrowserUtils::LastModifiedColName)
 			.OnSort(this, &SConcertSessionBrowser::OnColumnSortModeChanged)
-			.OnGetMenuContent_Lambda([this](){ return MakeHideColumnContextMenu(SessionHeaderRow.ToSharedRef(), ConcertBrowserUtils::LastModifiedColName); })
 			);
 
 	RestoreColumnVisibilityState(SessionHeaderRow.ToSharedRef(), InArgs._ColumnVisibilitySnapshot);
@@ -441,7 +448,7 @@ EColumnSortMode::Type SConcertSessionBrowser::GetColumnSortMode(const FName Colu
 	{
 		return PrimarySortMode;
 	}
-	else if (ColumnId == SecondarySortedColumn)
+	if (ColumnId == SecondarySortedColumn)
 	{
 		return SecondarySortMode;
 	}
@@ -496,7 +503,7 @@ void SConcertSessionBrowser::SortSessionList()
 		{
 			return true;
 		}
-		else if (Rhs->Type == FConcertSessionItem::EType::NewSession)
+		if (Rhs->Type == FConcertSessionItem::EType::NewSession)
 		{
 			return false;
 		}
@@ -505,23 +512,23 @@ void SConcertSessionBrowser::SortSessionList()
 		{
 			return SortMode == EColumnSortMode::Ascending ? Lhs->Type < Rhs->Type : Lhs->Type > Rhs->Type;
 		}
-		else if (ColName == ConcertBrowserUtils::SessionColName)
+		if (ColName == ConcertBrowserUtils::SessionColName)
 		{
 			return SortMode == EColumnSortMode::Ascending ? Lhs->SessionName < Rhs->SessionName : Lhs->SessionName > Rhs->SessionName;
 		}
-		else if (ColName == ConcertBrowserUtils::ServerColName)
+		if (ColName == ConcertBrowserUtils::ServerColName)
 		{
 			return SortMode == EColumnSortMode::Ascending ? Lhs->ServerName < Rhs->ServerName : Lhs->ServerName > Rhs->ServerName;
 		}
-		else if (ColName == ConcertBrowserUtils::ProjectColName)
+		if (ColName == ConcertBrowserUtils::ProjectColName)
 		{
 			return SortMode == EColumnSortMode::Ascending ? Lhs->ProjectName < Rhs->ProjectName : Lhs->ProjectName > Rhs->ProjectName;
 		}
-		else if (ColName == ConcertBrowserUtils::VersionColName)
+		if (ColName == ConcertBrowserUtils::VersionColName)
 		{
 			return SortMode == EColumnSortMode::Ascending ? Lhs->ProjectVersion < Rhs->ProjectVersion : Lhs->ProjectVersion > Rhs->ProjectVersion;
 		}
-		else if (ColName == ConcertBrowserUtils::LastModifiedColName)
+		if (ColName == ConcertBrowserUtils::LastModifiedColName)
 		{
 			return SortMode == EColumnSortMode::Ascending ? Lhs->LastModified < Rhs->LastModified : Lhs->LastModified > Rhs->LastModified;
 		}
@@ -539,7 +546,7 @@ void SConcertSessionBrowser::SortSessionList()
 		{
 			return true; // Lhs must be before Rhs based on the primary sort order.
 		}
-		else if (Compare(Rhs, Lhs, PrimarySortedColumn, PrimarySortMode)) // Invert operands order (goal is to check if operands are equal or not)
+		if (Compare(Rhs, Lhs, PrimarySortedColumn, PrimarySortMode)) // Invert operands order (goal is to check if operands are equal or not)
 		{
 			return false; // Rhs must be before Lhs based on the primary sort.
 		}
@@ -729,72 +736,8 @@ void SConcertSessionBrowser::RemoveSessionRow(const TSharedPtr<FConcertSessionIt
 	SessionsView->RequestListRefresh();
 }
 
-TSharedRef<SWidget> SConcertSessionBrowser::MakeSessionViewOptionsBar()
+TSharedRef<SWidget> SConcertSessionBrowser::MakeSessionTableFooter()
 {
-	auto AddFilterMenu = [this]()
-	{
-		FMenuBuilder MenuBuilder(/*bInShouldCloseWindowAfterMenuSelection=*/true, nullptr);
-
-		MenuBuilder.AddMenuEntry(
-			LOCTEXT("DisplayLastModifiedInRelativeTime", "Display Relative Time"),
-			TAttribute<FText>::CreateLambda([this]()
-			{
-				const bool bIsVisible = SessionHeaderRow->IsColumnVisible(ConcertBrowserUtils::LastModifiedColName);
-				return bIsVisible
-					? LOCTEXT("DisplayLastModifiedInRelativeTime.Tooltip.Visible", "Display the Last Modified column in relative time?")
-					: LOCTEXT("DisplayLastModifiedInRelativeTime.Tooltip.Hidden", "Disabled because the Last Modified column is hidden.");
-			}),
-			FSlateIcon(),
-			FUIAction(
-				FExecuteAction::CreateSP(this, &SConcertSessionBrowser::OnFilterMenuChecked, ConcertBrowserUtils::LastModifiedCheckBoxMenuName),
-				FCanExecuteAction::CreateLambda([this] { return SessionHeaderRow->IsColumnVisible(ConcertBrowserUtils::LastModifiedColName); }),
-				FIsActionChecked::CreateLambda([this] { return PersistentSettings->LastModifiedTimeFormat == ETimeFormat::Relative; })),
-			NAME_None,
-			EUserInterfaceActionType::ToggleButton
-		);
-		MenuBuilder.AddSeparator();
-		
-		MenuBuilder.AddMenuEntry(
-			LOCTEXT("ActiveSessions_Label", "Active Sessions"),
-			LOCTEXT("ActiveSessions_Tooltip", "Displays Active Sessions"),
-			FSlateIcon(),
-			FUIAction(
-				FExecuteAction::CreateSP(this, &SConcertSessionBrowser::OnFilterMenuChecked, ConcertBrowserUtils::ActiveSessionsCheckBoxMenuName),
-				FCanExecuteAction::CreateLambda([] { return true; }),
-				FIsActionChecked::CreateLambda([this] { return PersistentSettings->bShowActiveSessions; })),
-			NAME_None,
-			EUserInterfaceActionType::ToggleButton
-		);
-
-		MenuBuilder.AddMenuEntry(
-			LOCTEXT("ArchivedSessions_Label", "Archived Sessions"),
-			LOCTEXT("ArchivedSessions_Tooltip", "Displays Archived Sessions"),
-			FSlateIcon(),
-			FUIAction(
-				FExecuteAction::CreateSP(this, &SConcertSessionBrowser::OnFilterMenuChecked, ConcertBrowserUtils::ArchivedSessionsCheckBoxMenuName),
-				FCanExecuteAction::CreateLambda([] { return true; }),
-				FIsActionChecked::CreateLambda([this] { return PersistentSettings->bShowArchivedSessions; })),
-			NAME_None,
-			EUserInterfaceActionType::ToggleButton
-		);
-
-		MenuBuilder.AddMenuEntry(
-			LOCTEXT("DefaultServer_Label", "Default Server Sessions"),
-			LOCTEXT("DefaultServer_Tooltip", "Displays Sessions Hosted By the Default Server"),
-			FSlateIcon(),
-			FUIAction(
-				FExecuteAction::CreateSP(this, &SConcertSessionBrowser::OnFilterMenuChecked, ConcertBrowserUtils::DefaultServerCheckBoxMenuName),
-				FCanExecuteAction::CreateLambda([] { return true; }),
-				FIsActionChecked::CreateLambda([this] { return PersistentSettings->bShowDefaultServerSessionsOnly; })),
-			NAME_None,
-			EUserInterfaceActionType::ToggleButton
-		);
-
-		MenuBuilder.AddSeparator();
-		UE::ConcertSharedSlate::AddEntriesForShowingHiddenRows(SessionHeaderRow.ToSharedRef(), MenuBuilder);
-		return MenuBuilder.MakeWidget();
-	};
-
 	return SNew(SHorizontalBox)
 		+SHorizontalBox::Slot()
 		.AutoWidth()
@@ -840,37 +783,6 @@ TSharedRef<SWidget> SConcertSessionBrowser::MakeSessionViewOptionsBar()
 		.FillWidth(1.0)
 		[
 			SNew(SSpacer)
-		]
-
-		+SHorizontalBox::Slot()
-		.AutoWidth()
-		[
-			SNew(SComboButton)
-			.ComboButtonStyle(FAppStyle::Get(), "ComboButton")
-			.ForegroundColor(FStyleColors::Foreground)
-			.ContentPadding(0)
-			.OnGetMenuContent_Lambda(AddFilterMenu)
-			.HasDownArrow(true)
-			.ContentPadding(FMargin(1, 0))
-			.ButtonContent()
-			[
-				SNew(SHorizontalBox)
-
-				+SHorizontalBox::Slot()
-				.AutoWidth()
-				.VAlign(VAlign_Center)
-				[
-					SNew(SImage).Image(FAppStyle::Get().GetBrush("Icons.Visible")) // The eye ball image.
-				]
-
-				+SHorizontalBox::Slot()
-				.AutoWidth()
-				.Padding(2, 0, 0, 0)
-				.VAlign(VAlign_Center)
-				[
-					SNew(STextBlock).Text(LOCTEXT("ViewOptions", "View Options"))
-				]
-			]
 		];
 }
 
@@ -907,12 +819,13 @@ void SConcertSessionBrowser::OnFilterMenuChecked(const FName MenuName)
 	PersistentSettings->SaveConfig();
 }
 
-TSharedPtr<SWidget> SConcertSessionBrowser::MakeContextualMenu()
+TSharedPtr<SWidget> SConcertSessionBrowser::MakeTableContextualMenu()
 {
 	TArray<TSharedPtr<FConcertSessionItem>> SelectedItems = SessionsView->GetSelectedItems();
 	if (SelectedItems.Num() == 0 || (SelectedItems[0]->Type != FConcertSessionItem::EType::ActiveSession && SelectedItems[0]->Type != FConcertSessionItem::EType::ArchivedSession))
 	{
-		return nullptr; // No menu for editable rows.
+		// If clicking the header row, show the options for hiding them
+		return UE::ConcertSharedSlate::MakeTableContextMenu(SessionHeaderRow.ToSharedRef(), {}, true);
 	}
 
 	TSharedPtr<FConcertSessionItem> Item = SelectedItems[0];
@@ -980,6 +893,68 @@ TSharedPtr<SWidget> SConcertSessionBrowser::MakeContextualMenu()
 	return MenuBuilder.MakeWidget();
 }
 
+TSharedRef<SWidget> SConcertSessionBrowser::MakeViewOptionsComboButtonMenu()
+{
+	FMenuBuilder MenuBuilder(/*bInShouldCloseWindowAfterMenuSelection=*/true, nullptr);
+
+	MenuBuilder.AddMenuEntry(
+		LOCTEXT("DisplayLastModifiedInRelativeTime", "Display Relative Time"),
+		TAttribute<FText>::CreateLambda([this]()
+		{
+			const bool bIsVisible = SessionHeaderRow->IsColumnVisible(ConcertBrowserUtils::LastModifiedColName);
+			return bIsVisible
+				? LOCTEXT("DisplayLastModifiedInRelativeTime.Tooltip.Visible", "Display the Last Modified column in relative time?")
+				: LOCTEXT("DisplayLastModifiedInRelativeTime.Tooltip.Hidden", "Disabled because the Last Modified column is hidden.");
+		}),
+		FSlateIcon(),
+		FUIAction(
+			FExecuteAction::CreateSP(this, &SConcertSessionBrowser::OnFilterMenuChecked, ConcertBrowserUtils::LastModifiedCheckBoxMenuName),
+			FCanExecuteAction::CreateLambda([this] { return SessionHeaderRow->IsColumnVisible(ConcertBrowserUtils::LastModifiedColName); }),
+			FIsActionChecked::CreateLambda([this] { return PersistentSettings->LastModifiedTimeFormat == ETimeFormat::Relative; })),
+		NAME_None,
+		EUserInterfaceActionType::ToggleButton
+	);
+	MenuBuilder.AddSeparator();
+	
+	MenuBuilder.AddMenuEntry(
+		LOCTEXT("ActiveSessions_Label", "Active Sessions"),
+		LOCTEXT("ActiveSessions_Tooltip", "Displays Active Sessions"),
+		FSlateIcon(),
+		FUIAction(
+			FExecuteAction::CreateSP(this, &SConcertSessionBrowser::OnFilterMenuChecked, ConcertBrowserUtils::ActiveSessionsCheckBoxMenuName),
+			FCanExecuteAction::CreateLambda([] { return true; }),
+			FIsActionChecked::CreateLambda([this] { return PersistentSettings->bShowActiveSessions; })),
+		NAME_None,
+		EUserInterfaceActionType::ToggleButton
+	);
+
+	MenuBuilder.AddMenuEntry(
+		LOCTEXT("ArchivedSessions_Label", "Archived Sessions"),
+		LOCTEXT("ArchivedSessions_Tooltip", "Displays Archived Sessions"),
+		FSlateIcon(),
+		FUIAction(
+			FExecuteAction::CreateSP(this, &SConcertSessionBrowser::OnFilterMenuChecked, ConcertBrowserUtils::ArchivedSessionsCheckBoxMenuName),
+			FCanExecuteAction::CreateLambda([] { return true; }),
+			FIsActionChecked::CreateLambda([this] { return PersistentSettings->bShowArchivedSessions; })),
+		NAME_None,
+		EUserInterfaceActionType::ToggleButton
+	);
+
+	MenuBuilder.AddMenuEntry(
+		LOCTEXT("DefaultServer_Label", "Default Server Sessions"),
+		LOCTEXT("DefaultServer_Tooltip", "Displays Sessions Hosted By the Default Server"),
+		FSlateIcon(),
+		FUIAction(
+			FExecuteAction::CreateSP(this, &SConcertSessionBrowser::OnFilterMenuChecked, ConcertBrowserUtils::DefaultServerCheckBoxMenuName),
+			FCanExecuteAction::CreateLambda([] { return true; }),
+			FIsActionChecked::CreateLambda([this] { return PersistentSettings->bShowDefaultServerSessionsOnly; })),
+		NAME_None,
+		EUserInterfaceActionType::ToggleButton
+	);
+
+	return MenuBuilder.MakeWidget();
+}
+
 void SConcertSessionBrowser::OnSessionSelectionChanged(TSharedPtr<FConcertSessionItem> SelectedSession, ESelectInfo::Type SelectInfo)
 {
 	// Cancel editing the row to create, archive or restore a session (if any), unless the row was selected in code.
diff --git a/Engine/Plugins/Developer/Concert/ConcertUI/ConcertSharedSlate/Source/ConcertSharedSlate/Private/Session/History/SSessionHistory.cpp b/Engine/Plugins/Developer/Concert/ConcertUI/ConcertSharedSlate/Source/ConcertSharedSlate/Private/Session/History/SSessionHistory.cpp
--- a/Engine/Plugins/Developer/Concert/ConcertUI/ConcertSharedSlate/Source/ConcertSharedSlate/Private/Session/History/SSessionHistory.cpp
+++ b/Engine/Plugins/Developer/Concert/ConcertUI/ConcertSharedSlate/Source/ConcertSharedSlate/Private/Session/History/SSessionHistory.cpp
@@ -11,6 +11,7 @@
 #include "Session/Activity/SConcertSessionActivities.h"
 
 #include "Framework/MultiBox/MultiBoxBuilder.h"
+#include "Styling/ToolBarStyle.h"
 #include "Widgets/Input/SSearchBox.h"
 #include "Widgets/Layout/SSeparator.h"
 
@@ -77,21 +78,24 @@ void SSessionHistory::Construct(const FArguments& InArgs)
 		.UndoHistoryReflectionProvider(InArgs._UndoHistoryReflectionProvider)
 		.DarkenMutedActivities(InArgs._DarkenMutedActivities);
 
+	
+	TSharedPtr<SVerticalBox> ButtonExtension;
 	ChildSlot
 	[
 		SNew(SVerticalBox)
+
+		+SVerticalBox::Slot()
+		.AutoHeight()
+		[
+			SAssignNew(ButtonExtension, SVerticalBox)
+		]
+		
 		+SVerticalBox::Slot()
 		.AutoHeight()
-		.Padding(1, 1)
+		.Padding(2, 0, 2, 4)
 		[
 			SNew(SHorizontalBox)
 
-			+SHorizontalBox::Slot()
-			.AutoWidth()
-			[
-				InArgs._SearchButtonArea.Widget
-			]
-
 			+SHorizontalBox::Slot()
 			.FillWidth(1.f)
 			[
@@ -101,6 +105,12 @@ void SSessionHistory::Construct(const FArguments& InArgs)
 				.OnTextCommitted(this, &SSessionHistory::OnSearchTextCommitted)
 				.DelayChangeNotificationsWhileTyping(true)
 			]
+			
+			+SHorizontalBox::Slot()
+			.AutoWidth()
+			[
+				ActivityListViewOptions->MakeViewOptionsComboButton()
+			]
 		]
 
 		+SVerticalBox::Slot()
@@ -121,14 +131,29 @@ void SSessionHistory::Construct(const FArguments& InArgs)
 		[
 			ActivityListViewOptions->MakeStatusBar(
 				TAttribute<int32>(ActivityListView.Get(), &SConcertSessionActivities::GetTotalActivityNum),
-				TAttribute<int32>(ActivityListView.Get(), &SConcertSessionActivities::GetDisplayedActivityNum),
-				FConcertSessionActivitiesOptions::FExtendContextMenu::CreateLambda([this](FMenuBuilder& MenuBuilder)
-				{
-					MenuBuilder.AddSeparator();
-					AddEntriesForShowingHiddenRows(ActivityListView->GetHeaderRow().ToSharedRef(), MenuBuilder);
-				}))
+				TAttribute<int32>(ActivityListView.Get(), &SConcertSessionActivities::GetDisplayedActivityNum)
+				)
 		]
 	];
+
+	const bool bHasExtensions = InArgs._SearchButtonArea.Widget != SNullWidget::NullWidget;
+	if (bHasExtensions)
+	{
+		ButtonExtension->AddSlot()
+			.AutoHeight()
+			[
+				InArgs._SearchButtonArea.Widget
+			];
+		
+		ButtonExtension->AddSlot()
+			.AutoHeight()
+			.Padding(0, 3)
+			[
+				SNew(SSeparator)
+				.Thickness(2.f)
+				.SeparatorImage(&FCoreStyle::Get().GetWidgetStyle<FToolBarStyle>("ToolBar").SeparatorBrush)
+			];
+	}
 }
 
 void SSessionHistory::ReloadActivities(TMap<FGuid, FConcertClientInfo> InEndpointClientInfoMap, TArray<FConcertSessionActivity> InFetchedActivities)
diff --git a/Engine/Plugins/Developer/Concert/ConcertUI/ConcertSharedSlate/Source/ConcertSharedSlate/Public/ConcertFrontendUtils.h b/Engine/Plugins/Developer/Concert/ConcertUI/ConcertSharedSlate/Source/ConcertSharedSlate/Public/ConcertFrontendUtils.h
--- a/Engine/Plugins/Developer/Concert/ConcertUI/ConcertSharedSlate/Source/ConcertSharedSlate/Public/ConcertFrontendUtils.h
+++ b/Engine/Plugins/Developer/Concert/ConcertUI/ConcertSharedSlate/Source/ConcertSharedSlate/Public/ConcertFrontendUtils.h
@@ -5,6 +5,8 @@
 #include "CoreMinimal.h"
 #include "ConcertFrontendStyle.h"
 #include "Styling/AppStyle.h"
+#include "Widgets/Images/SImage.h"
+#include "Widgets/Input/SComboButton.h"
 #include "Widgets/Layout/SBorder.h"
 #include "Widgets/Text/STextBlock.h"
 #include "Widgets/Layout/SExpandableArea.h"
@@ -145,6 +147,20 @@ namespace ConcertFrontendUtils
 			? FormatRelativeTime(Time, CurrTime)
 			: FText::AsDateTime(Time);
 	}
+
+	static TSharedRef<SComboButton> CreateViewOptionsComboButton(FOnGetContent GetMenuContentDelegate)
+	{
+		return SNew(SComboButton)
+			.ComboButtonStyle(FAppStyle::Get(), "SimpleComboButtonWithIcon")
+			.OnGetMenuContent(MoveTemp(GetMenuContentDelegate))
+			.HasDownArrow(false)
+			.ButtonContent()
+			[
+				SNew(SImage)
+				.ColorAndOpacity(FSlateColor::UseForeground())
+				.Image(FAppStyle::Get().GetBrush("Icons.Settings"))
+			];
+	}
 };
 
 #undef LOCTEXT_NAMESPACE
diff --git a/Engine/Plugins/Developer/Concert/ConcertUI/ConcertSharedSlate/Source/ConcertSharedSlate/Public/ConcertHeaderRowUtils.h b/Engine/Plugins/Developer/Concert/ConcertUI/ConcertSharedSlate/Source/ConcertSharedSlate/Public/ConcertHeaderRowUtils.h
--- a/Engine/Plugins/Developer/Concert/ConcertUI/ConcertSharedSlate/Source/ConcertSharedSlate/Public/ConcertHeaderRowUtils.h
+++ b/Engine/Plugins/Developer/Concert/ConcertUI/ConcertSharedSlate/Source/ConcertSharedSlate/Public/ConcertHeaderRowUtils.h
@@ -22,16 +22,8 @@ namespace UE::ConcertSharedSlate
 {
 	DECLARE_DELEGATE_OneParam(FSaveColumnVisibilitySnapshot, const FColumnVisibilitySnapshot& /*Snapshot*/);
 
-	// TODO DP: Remove MakeHideColumnContextMenu and AddHideColumnEntry once the UX pass is done
-	
-	/**
-	 * Creates a widget for hiding a SHeaderRow::FColumn when it is right-clicked.
-	 */
-	CONCERTSHAREDSLATE_API TSharedRef<SWidget> MakeHideColumnContextMenu(const TSharedRef<SHeaderRow>& HeaderRow, const FName ForColumnID);
-	CONCERTSHAREDSLATE_API void AddHideColumnEntry(const TSharedRef<SHeaderRow>& HeaderRow, const FName ForColumnID, FMenuBuilder& MenuBuilder);
-
 	/** Util function that can be directly fed into SListView::OnContextMenuOpening */
-	CONCERTSHAREDSLATE_API TSharedRef<SWidget> MakeTableContextMenu(const TSharedRef<SHeaderRow>& HeaderRow, TMap<FName, bool> ColumnsVisibleByDefault, bool bDefaultVisibility = true);
+	CONCERTSHAREDSLATE_API TSharedRef<SWidget> MakeTableContextMenu(const TSharedRef<SHeaderRow>& HeaderRow, TMap<FName, bool> ColumnsVisibleByDefault = {}, bool bDefaultVisibility = true);
 
 	/**
 	 * Adds an entry for "Show all", "Hide all", and "Restore column visibility".
diff --git a/Engine/Plugins/Developer/Concert/ConcertUI/ConcertSharedSlate/Source/ConcertSharedSlate/Public/Session/Activity/SConcertSessionActivities.h b/Engine/Plugins/Developer/Concert/ConcertUI/ConcertSharedSlate/Source/ConcertSharedSlate/Public/Session/Activity/SConcertSessionActivities.h
--- a/Engine/Plugins/Developer/Concert/ConcertUI/ConcertSharedSlate/Source/ConcertSharedSlate/Public/Session/Activity/SConcertSessionActivities.h
+++ b/Engine/Plugins/Developer/Concert/ConcertUI/ConcertSharedSlate/Source/ConcertSharedSlate/Public/Session/Activity/SConcertSessionActivities.h
@@ -362,32 +362,23 @@ public:
 	DECLARE_DELEGATE_OneParam(FExtendContextMenu, FMenuBuilder&)
 
 	/** Makes a status bar widget displaying the activity shown and the standard view options button. */
-	TSharedRef<SWidget> MakeStatusBar(TAttribute<int32> Total, TAttribute<int32> Displayed, TOptional<FExtendContextMenu> ExtendMenu = {});
-	
-	/** Makes a widgets saying how many operations are shown with respect to those filtered out. */
-	TSharedRef<SWidget> MakeDisplayedActivityCountWidget(TAttribute<int32> Total, TAttribute<int32> Displayed);
+	TSharedRef<SWidget> MakeStatusBar(TAttribute<int32> Total, TAttribute<int32> Displayed);
 	
 	/** Makes a standard View Options widget, displaying the eye ball icon and showing the possible options. */
-	TSharedRef<SWidget> MakeViewOptionsWidget(TOptional<FExtendContextMenu> ExtendMenu = {});
-	
+	TSharedRef<SWidget> MakeViewOptionsComboButton(TOptional<FExtendContextMenu> ExtendMenu = {});
 	/** Returns a Menu widgets containing the available options. */
 	TSharedRef<SWidget> MakeMenuWidget(TOptional<FExtendContextMenu> ExtendMenu = {});
 
 	/** Returns the time format option. */
 	ETimeFormat GetTimeFormat() const { return bDisplayRelativeTime ? ETimeFormat::Relative : ETimeFormat::Absolute; }
-
 	/** Returns whether the connection activities are listed in SConcertSessionActivities. */
 	EVisibility GetConnectionActivitiesVisibility() const { return bDisplayConnectionActivities ? EVisibility::Visible : EVisibility::Hidden; }
-
 	/** Returns whether the lock activities are listed in SConcertSessionActivities. */
 	EVisibility GetLockActivitiesVisibility() const { return bDisplayLockActivities ? EVisibility::Visible : EVisibility::Hidden; }
-
 	/** Returns whether the package activities are listed in SConcertSessionActivities. */
 	EVisibility GetPackageActivitiesVisibility() const { return bDisplayPackageActivities ? EVisibility::Visible : EVisibility::Hidden; }
-
 	/** Returns whether the transaction activities are listed in SConcertSessionActivities. */
 	EVisibility GetTransactionActivitiesVisibility() const { return bDisplayTransactionActivities ? EVisibility::Visible : EVisibility::Hidden; }
-
 	/** Returns whether the ignored activities are listed in SConcertSessionActivities. */
 	EVisibility GetIgnoredActivitiesVisibility() const { return bDisplayIgnoredActivities ? EVisibility::Visible : EVisibility::Hidden; }
 
@@ -396,34 +387,24 @@ public:
 
 	/** Enables the 'connection activity' filter check box. (Show Connection Activities). */
 	bool bEnableConnectionActivityFiltering = true;
-
 	/** Enables the 'lock activity' filter check box. (Show Lock Activities).*/
 	bool bEnableLockActivityFiltering = true;
-
 	/** Enables the 'package activity' filter check box (Show Package Activities). */
 	bool bEnablePackageActivityFiltering = true;
-
 	/** Enables the 'transaction activity' filter check box (Show Transaction Activities). */
 	bool bEnableTransactionActivityFiltering = true;
-
 	/** Enables the 'ignored activity' filter check box (Show Unrecoverable Activities). */
 	bool bEnableIgnoredActivityFiltering = false;
-
 	/** Controls whether the time is displayed as absolute or relative. (Display Relative Time). */
 	bool bDisplayRelativeTime = true;
-
 	/** If connection filtering is enabled, controls whether connection activities are filtered out. */
 	bool bDisplayConnectionActivities = false;
-
 	/** If lock filtering is enabled, controls whether lock activities are filtered out. */
 	bool bDisplayLockActivities = false;
-
 	/** If package filtering is enabled, controls whether package activities are filtered out. */
 	bool bDisplayPackageActivities = true;
-
 	/** If transaction filtering is enabled, controls whether transaction activities are filtered out.*/
 	bool bDisplayTransactionActivities = true;
-
 	/** If ignored activity filtering is enabled, controls whether ignored activities are filtered out.*/
 	bool bDisplayIgnoredActivities = false;
 };
diff --git a/Engine/Plugins/Developer/Concert/ConcertUI/ConcertSharedSlate/Source/ConcertSharedSlate/Public/Session/Browser/SConcertSessionBrowser.h b/Engine/Plugins/Developer/Concert/ConcertUI/ConcertSharedSlate/Source/ConcertSharedSlate/Public/Session/Browser/SConcertSessionBrowser.h
--- a/Engine/Plugins/Developer/Concert/ConcertUI/ConcertSharedSlate/Source/ConcertSharedSlate/Public/Session/Browser/SConcertSessionBrowser.h
+++ b/Engine/Plugins/Developer/Concert/ConcertUI/ConcertSharedSlate/Source/ConcertSharedSlate/Public/Session/Browser/SConcertSessionBrowser.h
@@ -126,7 +126,7 @@ private:
 	TSharedRef<SWidget> MakeControlBar(const FArguments& InArgs);
 	TSharedRef<SWidget> MakeButtonBar(const FArguments& InArgs);
 	TSharedRef<SWidget> MakeSessionTableView(const FArguments& InArgs);
-	TSharedRef<SWidget> MakeSessionViewOptionsBar();
+	TSharedRef<SWidget> MakeSessionTableFooter();
 	TSharedRef<ITableRow> OnGenerateSessionRowWidget(TSharedPtr<FConcertSessionItem> Item, const TSharedRef<STableViewBase>& OwnerTable);
 
 	// Creates row widgets for session list view, validates user inputs and forward user requests for processing to a delegate function implemented by this class.
@@ -137,7 +137,8 @@ private:
 	TSharedRef<ITableRow> MakeSaveSessionRowWidget(const TSharedPtr<FConcertSessionItem>& ArchivedItem, const TSharedRef<STableViewBase>& OwnerTable);
 
 	// Creates the contextual menu when right clicking a session list view row.
-	TSharedPtr<SWidget> MakeContextualMenu();
+	TSharedPtr<SWidget> MakeTableContextualMenu();
+	TSharedRef<SWidget> MakeViewOptionsComboButtonMenu();
 
 	// The buttons above the session view.
 	bool IsNewButtonEnabledInternal() const;
