Group archived and live sessions

diff --git a/Engine/Plugins/Developer/Concert/ConcertApp/MultiUserClient/Source/MultiUserClient/Private/Widgets/SConcertBrowser.cpp b/Engine/Plugins/Developer/Concert/ConcertApp/MultiUserClient/Source/MultiUserClient/Private/Widgets/SConcertBrowser.cpp
--- a/Engine/Plugins/Developer/Concert/ConcertApp/MultiUserClient/Source/MultiUserClient/Private/Widgets/SConcertBrowser.cpp
+++ b/Engine/Plugins/Developer/Concert/ConcertApp/MultiUserClient/Source/MultiUserClient/Private/Widgets/SConcertBrowser.cpp
@@ -18,7 +18,7 @@
 #include "Algo/TiedTupleOutput.h"
 
 #include "Session/Browser/ConcertBrowserUtils.h"
-#include "Session/Browser/ConcertSessionItem.h"
+#include "Session/Browser/Items/ConcertSessionTreeItem.h"
 #include "Session/Browser/IConcertSessionBrowserController.h"
 #include "Session/Browser/SConcertSessionBrowser.h"
 
@@ -1013,14 +1013,14 @@ private:
 	// Layout the 'session|details' split view.
 	TSharedRef<SWidget> MakeBrowserContent(TSharedPtr<FText> InSearchText);
 	void ExtendControlButtons(FExtender& Extender);
-	void ExtendSessionContextMenu(const TSharedPtr<FConcertSessionItem>& Item, FExtender& Extender);
+	void ExtendSessionContextMenu(const TSharedPtr<FConcertSessionTreeItem>& Item, FExtender& Extender);
 	TSharedRef<SWidget> MakeUserAndSettings();
 	TSharedRef<SWidget> MakeOverlayedTableView(const TSharedRef<SWidget>& TableView);
 	
 	// Layouts the session detail panel.
-	TSharedRef<SWidget> MakeSessionDetails(TSharedPtr<FConcertSessionItem> Item);
-	TSharedRef<SWidget> MakeActiveSessionDetails(TSharedPtr<FConcertSessionItem> Item);
-	TSharedRef<SWidget> MakeArchivedSessionDetails(TSharedPtr<FConcertSessionItem> Item);
+	TSharedRef<SWidget> MakeSessionDetails(TSharedPtr<FConcertSessionTreeItem> Item);
+	TSharedRef<SWidget> MakeActiveSessionDetails(TSharedPtr<FConcertSessionTreeItem> Item);
+	TSharedRef<SWidget> MakeArchivedSessionDetails(TSharedPtr<FConcertSessionTreeItem> Item);
 	TSharedRef<ITableRow> OnGenerateClientRowWidget(TSharedPtr<FConcertSessionClientInfo> Item, const TSharedRef<STableViewBase>& OwnerTable);
 	void PopulateSessionInfoGrid(SGridPanel& Grid, const FConcertSessionInfo& SessionInfo);
 
@@ -1037,12 +1037,12 @@ private:
 
 	// SConcertSessionBrowser extensions
 	FReply OnJoinButtonClicked();
-	void RequestJoinSession(const TSharedPtr<FConcertSessionItem>& ActiveItem);
+	void RequestJoinSession(const TSharedPtr<FConcertSessionTreeItem>& ActiveItem);
 
 	// Manipulates the sessions view (the array and the UI).
-	void OnSessionSelectionChanged(const TSharedPtr<FConcertSessionItem>& SessionItem);
-	void OnSessionDoubleClicked(const TSharedPtr<FConcertSessionItem>& SessionItem);
-	bool ConfirmDeleteSessionWithDialog(const TArray<TSharedPtr<FConcertSessionItem>>& SessionItems) const;
+	void OnSessionSelectionChanged(const TSharedPtr<FConcertSessionTreeItem>& SessionItem);
+	void OnSessionDoubleClicked(const TSharedPtr<FConcertSessionTreeItem>& SessionItem);
+	bool ConfirmDeleteSessionWithDialog(const TArray<TSharedPtr<FConcertSessionTreeItem>>& SessionItems) const;
 
 	// Update server/session/clients lists.
 	EActiveTimerReturnType TickDiscovery(double InCurrentTime, float InDeltaTime);
@@ -1296,9 +1296,9 @@ void SConcertClientSessionBrowser::ExtendControlButtons(FExtender& Extender)
 		}));
 }
 
-void SConcertClientSessionBrowser::ExtendSessionContextMenu(const TSharedPtr<FConcertSessionItem>& Item, FExtender& Extender)
+void SConcertClientSessionBrowser::ExtendSessionContextMenu(const TSharedPtr<FConcertSessionTreeItem>& Item, FExtender& Extender)
 {
-	if (Item->Type == FConcertSessionItem::EType::ActiveSession)
+	if (Item->Type == FConcertSessionTreeItem::EType::ActiveSession)
 	{
 		Extender.AddMenuExtension(
 			SConcertSessionBrowser::SessionContextMenuExtensionHooks::ManageSession,
@@ -1403,8 +1403,8 @@ void SConcertClientSessionBrowser::UpdateDiscovery()
 	}
 
 	// If an active session is selected.
-	TArray<TSharedPtr<FConcertSessionItem>> SelectedSession = SessionBrowser->GetSelectedItems();
-	if (SelectedSession.Num() && SelectedSession[0]->Type == FConcertSessionItem::EType::ActiveSession)
+	TArray<TSharedPtr<FConcertSessionTreeItem>> SelectedSession = SessionBrowser->GetSelectedItems();
+	if (SelectedSession.Num() && SelectedSession[0]->Type == FConcertSessionTreeItem::EType::ActiveSession)
 	{
 		// Ensure to poll its clients.
 		uint32 CachedClientListVersion = Controller->TickClientsDiscovery(SelectedSession[0]->ServerAdminEndpointId, SelectedSession[0]->SessionId);
@@ -1448,17 +1448,17 @@ void SConcertClientSessionBrowser::RefreshClientList(const TArray<FConcertSessio
 	ClientsView->RequestListRefresh();
 }
 
-TSharedRef<SWidget> SConcertClientSessionBrowser::MakeSessionDetails(TSharedPtr<FConcertSessionItem> Item)
+TSharedRef<SWidget> SConcertClientSessionBrowser::MakeSessionDetails(TSharedPtr<FConcertSessionTreeItem> Item)
 {
 	const FConcertSessionInfo* SessionInfo = nullptr;
 
 	if (Item.IsValid())
 	{
-		if (Item->Type == FConcertSessionItem::EType::ActiveSession || Item->Type == FConcertSessionItem::EType::SaveSession)
+		if (Item->Type == FConcertSessionTreeItem::EType::ActiveSession || Item->Type == FConcertSessionTreeItem::EType::SaveSession)
 		{
 			return MakeActiveSessionDetails(Item);
 		}
-		else if (Item->Type == FConcertSessionItem::EType::ArchivedSession)
+		else if (Item->Type == FConcertSessionTreeItem::EType::ArchivedSession)
 		{
 			return MakeArchivedSessionDetails(Item);
 		}
@@ -1467,7 +1467,7 @@ TSharedRef<SWidget> SConcertClientSessionBrowser::MakeSessionDetails(TSharedPtr<
 	return NoSessionSelectedPanel.ToSharedRef();
 }
 
-TSharedRef<SWidget> SConcertClientSessionBrowser::MakeActiveSessionDetails(TSharedPtr<FConcertSessionItem> Item)
+TSharedRef<SWidget> SConcertClientSessionBrowser::MakeActiveSessionDetails(TSharedPtr<FConcertSessionTreeItem> Item)
 {
 	const TOptional<FConcertSessionInfo> SessionInfo = Controller->GetActiveSessionInfo(Item->ServerAdminEndpointId, Item->SessionId);
 	if (!SessionInfo)
@@ -1580,7 +1580,7 @@ TSharedRef<SWidget> SConcertClientSessionBrowser::MakeActiveSessionDetails(TShar
 	return Widget;
 }
 
-TSharedRef<SWidget> SConcertClientSessionBrowser::MakeArchivedSessionDetails(TSharedPtr<FConcertSessionItem> Item)
+TSharedRef<SWidget> SConcertClientSessionBrowser::MakeArchivedSessionDetails(TSharedPtr<FConcertSessionTreeItem> Item)
 {
 	const TOptional<FConcertSessionInfo> SessionInfo = Controller->GetArchivedSessionInfo(Item->ServerAdminEndpointId, Item->SessionId);
 	if (!SessionInfo)
@@ -1686,7 +1686,7 @@ TSharedRef<ITableRow> SConcertClientSessionBrowser::OnGenerateClientRowWidget(TS
 	];
 }
 
-void SConcertClientSessionBrowser::OnSessionSelectionChanged(const TSharedPtr<FConcertSessionItem>& SelectedSession)
+void SConcertClientSessionBrowser::OnSessionSelectionChanged(const TSharedPtr<FConcertSessionTreeItem>& SelectedSession)
 {
 	// Clear the list of clients (if any)
 	Clients.Reset();
@@ -1695,14 +1695,14 @@ void SConcertClientSessionBrowser::OnSessionSelectionChanged(const TSharedPtr<FC
 	SessionDetailsView->SetContent(MakeSessionDetails(SelectedSession));
 }
 
-void SConcertClientSessionBrowser::OnSessionDoubleClicked(const TSharedPtr<FConcertSessionItem>& SelectedSession)
+void SConcertClientSessionBrowser::OnSessionDoubleClicked(const TSharedPtr<FConcertSessionTreeItem>& SelectedSession)
 {
 	switch (SelectedSession->Type)
 	{
-	case FConcertSessionItem::EType::ActiveSession: 
+	case FConcertSessionTreeItem::EType::ActiveSession: 
 		RequestJoinSession(SelectedSession);
 		break;
-	case FConcertSessionItem::EType::ArchivedSession:
+	case FConcertSessionTreeItem::EType::ArchivedSession:
 		SessionBrowser->InsertRestoreSessionAsEditableRow(SelectedSession);
 		break;
 	default: 
@@ -1710,12 +1710,12 @@ void SConcertClientSessionBrowser::OnSessionDoubleClicked(const TSharedPtr<FConc
 	}
 }
 
-bool SConcertClientSessionBrowser::ConfirmDeleteSessionWithDialog(const TArray<TSharedPtr<FConcertSessionItem>>& SessionItems) const
+bool SConcertClientSessionBrowser::ConfirmDeleteSessionWithDialog(const TArray<TSharedPtr<FConcertSessionTreeItem>>& SessionItems) const
 {
 	TSet<FString> SessionNames;
 	TSet<FString> UniqueServers;
 
-	Algo::Transform(SessionItems, Algo::TieTupleAdd(SessionNames, UniqueServers), [](const TSharedPtr<FConcertSessionItem>& Item)
+	Algo::Transform(SessionItems, Algo::TieTupleAdd(SessionNames, UniqueServers), [](const TSharedPtr<FConcertSessionTreeItem>& Item)
 	{
 		return MakeTuple(Item->SessionName, Item->ServerName);
 	});
@@ -1738,8 +1738,8 @@ bool SConcertClientSessionBrowser::IsLaunchServerButtonEnabled() const
 
 bool SConcertClientSessionBrowser::IsJoinButtonEnabled() const
 {
-	TArray<TSharedPtr<FConcertSessionItem>> SelectedItems = SessionBrowser->GetSelectedItems();
-	return SelectedItems.Num() == 1 && SelectedItems[0]->Type == FConcertSessionItem::EType::ActiveSession;
+	TArray<TSharedPtr<FConcertSessionTreeItem>> SelectedItems = SessionBrowser->GetSelectedItems();
+	return SelectedItems.Num() == 1 && SelectedItems[0]->Type == FConcertSessionTreeItem::EType::ActiveSession;
 }
 
 bool SConcertClientSessionBrowser::IsAutoJoinButtonEnabled() const
@@ -1791,7 +1791,7 @@ FReply SConcertClientSessionBrowser::OnCancelAutoJoinButtonClicked()
 
 FReply SConcertClientSessionBrowser::OnJoinButtonClicked()
 {
-	TArray<TSharedPtr<FConcertSessionItem>> SelectedItems = SessionBrowser->GetSelectedItems();
+	TArray<TSharedPtr<FConcertSessionTreeItem>> SelectedItems = SessionBrowser->GetSelectedItems();
 	if (SelectedItems.Num() == 1)
 	{
 		RequestJoinSession(SelectedItems[0]);
@@ -1799,7 +1799,7 @@ FReply SConcertClientSessionBrowser::OnJoinButtonClicked()
 	return FReply::Handled();
 }
 
-void SConcertClientSessionBrowser::RequestJoinSession(const TSharedPtr<FConcertSessionItem>& LiveItem)
+void SConcertClientSessionBrowser::RequestJoinSession(const TSharedPtr<FConcertSessionTreeItem>& LiveItem)
 {
 	Controller->JoinSession(LiveItem->ServerAdminEndpointId, LiveItem->SessionId);
 }
diff --git a/Engine/Plugins/Developer/Concert/ConcertApp/MultiUserServer/Source/MultiUserServer/Private/Widgets/Browser/ConcertServerSessionBrowserController.cpp b/Engine/Plugins/Developer/Concert/ConcertApp/MultiUserServer/Source/MultiUserServer/Private/Widgets/Browser/ConcertServerSessionBrowserController.cpp
--- a/Engine/Plugins/Developer/Concert/ConcertApp/MultiUserServer/Source/MultiUserServer/Private/Widgets/Browser/ConcertServerSessionBrowserController.cpp
+++ b/Engine/Plugins/Developer/Concert/ConcertApp/MultiUserServer/Source/MultiUserServer/Private/Widgets/Browser/ConcertServerSessionBrowserController.cpp
@@ -6,10 +6,11 @@
 #include "ConcertServerStyle.h"
 #include "IConcertServer.h"
 #include "IConcertSyncServer.h"
-#include "Framework/Docking/TabManager.h"
 #include "SConcertServerSessionBrowser.h"
+#include "Session/Browser/Items/ConcertSessionTreeItem.h"
+
+#include "Framework/Docking/TabManager.h"
 #include "Misc/AsyncTaskNotification.h"
-#include "Session/Browser/ConcertSessionItem.h"
 #include "Textures/SlateIcon.h"
 #include "Widgets/Docking/SDockTab.h"
 #include "Widgets/Notifications/SNotificationList.h"
@@ -271,7 +272,7 @@ namespace UE::MultiUserServer
 		}
 	}
 
-	void FConcertServerSessionBrowserController::OpenSession(const TSharedPtr<FConcertSessionItem>& SessionItem)
+	void FConcertServerSessionBrowserController::OpenSession(const TSharedPtr<FConcertSessionTreeItem>& SessionItem)
 	{
 		Owner.Pin()->OpenSessionTab(SessionItem->SessionId);
 	}
diff --git a/Engine/Plugins/Developer/Concert/ConcertApp/MultiUserServer/Source/MultiUserServer/Private/Widgets/Browser/ConcertServerSessionBrowserController.h b/Engine/Plugins/Developer/Concert/ConcertApp/MultiUserServer/Source/MultiUserServer/Private/Widgets/Browser/ConcertServerSessionBrowserController.h
--- a/Engine/Plugins/Developer/Concert/ConcertApp/MultiUserServer/Source/MultiUserServer/Private/Widgets/Browser/ConcertServerSessionBrowserController.h
+++ b/Engine/Plugins/Developer/Concert/ConcertApp/MultiUserServer/Source/MultiUserServer/Private/Widgets/Browser/ConcertServerSessionBrowserController.h
@@ -9,7 +9,7 @@
 
 class IConcertServerSession;
 class IConcertServer;
-class FConcertSessionItem;
+class FConcertSessionTreeItem;
 class FSpawnTabArgs;
 class SDockTab;
 
@@ -75,7 +75,7 @@ namespace UE::MultiUserServer
 		void RefreshSessionList();
 
 		// Session actions
-		void OpenSession(const TSharedPtr<FConcertSessionItem>& SessionItem);
+		void OpenSession(const TSharedPtr<FConcertSessionTreeItem>& SessionItem);
 		void RenameActiveSessionInternal(const FGuid& ServerAdminEndpointId, const FGuid& SessionId, const FString& NewName);
 		void RenameArchivedSessionInternal(const FGuid& ServerAdminEndpointId, const FGuid& SessionId, const FString& NewName);
 		void RenameSessionInternal(const FGuid& SessionId, const FString& NewName, const FConcertSessionInfo& SessionInfo);
diff --git a/Engine/Plugins/Developer/Concert/ConcertApp/MultiUserServer/Source/MultiUserServer/Private/Widgets/Browser/SConcertServerSessionBrowser.cpp b/Engine/Plugins/Developer/Concert/ConcertApp/MultiUserServer/Source/MultiUserServer/Private/Widgets/Browser/SConcertServerSessionBrowser.cpp
--- a/Engine/Plugins/Developer/Concert/ConcertApp/MultiUserServer/Source/MultiUserServer/Private/Widgets/Browser/SConcertServerSessionBrowser.cpp
+++ b/Engine/Plugins/Developer/Concert/ConcertApp/MultiUserServer/Source/MultiUserServer/Private/Widgets/Browser/SConcertServerSessionBrowser.cpp
@@ -4,7 +4,7 @@
 
 #include "MultiUserServerModule.h"
 #include "Session/Browser/ConcertBrowserUtils.h"
-#include "Session/Browser/ConcertSessionItem.h"
+#include "Session/Browser/Items/ConcertSessionTreeItem.h"
 #include "Session/Browser/SConcertSessionBrowser.h"
 #include "Settings/MultiUserServerColumnVisibilitySettings.h"
 #include "Widgets/Browser/ConcertServerSessionBrowserController.h"
@@ -46,7 +46,7 @@ namespace UE::MultiUserServer
 			});
 	}
 
-	void SConcertServerSessionBrowser::RequestDeleteSession(const TArray<TSharedPtr<FConcertSessionItem>>& SessionItems) 
+	void SConcertServerSessionBrowser::RequestDeleteSession(const TArray<TSharedPtr<FConcertSessionTreeItem>>& SessionItems) 
 	{
 		const FText Message = [this, &SessionItems]()
 		{
@@ -60,13 +60,13 @@ namespace UE::MultiUserServer
 				
 			switch (SessionItems[0]->Type)
 			{
-			case FConcertSessionItem::EType::ActiveSession:
+			case FConcertSessionTreeItem::EType::ActiveSession:
 				return FText::Format(
 					LOCTEXT("DeletedActiveDescription", "There {0}|plural(one=is,other=are) {0} connected {0}|plural(one=client,other=clients) in the current session.\nDeleting a session will force all connected clients to disconnect.\n\nDelete {1}?"),
 					Controller.Pin()->GetNumConnectedClients(SessionItems[0]->SessionId),
 					FText::FromString(SessionItems[0]->SessionName)
 				);
-			case FConcertSessionItem::EType::ArchivedSession:
+			case FConcertSessionTreeItem::EType::ArchivedSession:
 				return FText::Format(
 					LOCTEXT("DeleteArchivedDescription", "Deleting a session will cause all associated data to be removed.\n\nDelete {0}?"),
 					FText::FromString(SessionItems[0]->SessionName)
@@ -79,7 +79,7 @@ namespace UE::MultiUserServer
 		DeleteSessionsWithFakeModalQuestion(Message, SessionItems);
 	}
 
-	void SConcertServerSessionBrowser::DeleteSessionsWithFakeModalQuestion(const FText& Message, const TArray<TSharedPtr<FConcertSessionItem>>& SessionItems)
+	void SConcertServerSessionBrowser::DeleteSessionsWithFakeModalQuestion(const FText& Message, const TArray<TSharedPtr<FConcertSessionTreeItem>>& SessionItems)
 	{
 		auto DeleteArchived = [WeakController = TWeakPtr<FConcertServerSessionBrowserController>(Controller), SessionItems]()
 		{
diff --git a/Engine/Plugins/Developer/Concert/ConcertApp/MultiUserServer/Source/MultiUserServer/Private/Widgets/Browser/SConcertServerSessionBrowser.h b/Engine/Plugins/Developer/Concert/ConcertApp/MultiUserServer/Source/MultiUserServer/Private/Widgets/Browser/SConcertServerSessionBrowser.h
--- a/Engine/Plugins/Developer/Concert/ConcertApp/MultiUserServer/Source/MultiUserServer/Private/Widgets/Browser/SConcertServerSessionBrowser.h
+++ b/Engine/Plugins/Developer/Concert/ConcertApp/MultiUserServer/Source/MultiUserServer/Private/Widgets/Browser/SConcertServerSessionBrowser.h
@@ -49,7 +49,7 @@ namespace UE::MultiUserServer
 	
 		TSharedRef<SWidget> MakeSessionTableView(const FArguments& InArgs);
 
-		void RequestDeleteSession(const TArray<TSharedPtr<FConcertSessionItem>>& SessionItems);
-		void DeleteSessionsWithFakeModalQuestion(const FText& Message, const TArray<TSharedPtr<FConcertSessionItem>>& SessionItems);
+		void RequestDeleteSession(const TArray<TSharedPtr<FConcertSessionTreeItem>>& SessionItems);
+		void DeleteSessionsWithFakeModalQuestion(const FText& Message, const TArray<TSharedPtr<FConcertSessionTreeItem>>& SessionItems);
 	};
 }
diff --git a/Engine/Plugins/Developer/Concert/ConcertApp/MultiUserServer/Source/MultiUserServer/Private/Widgets/SessionTabs/Archived/ArchivedConcertSessionTab.cpp b/Engine/Plugins/Developer/Concert/ConcertApp/MultiUserServer/Source/MultiUserServer/Private/Widgets/SessionTabs/Archived/ArchivedConcertSessionTab.cpp
--- a/Engine/Plugins/Developer/Concert/ConcertApp/MultiUserServer/Source/MultiUserServer/Private/Widgets/SessionTabs/Archived/ArchivedConcertSessionTab.cpp
+++ b/Engine/Plugins/Developer/Concert/ConcertApp/MultiUserServer/Source/MultiUserServer/Private/Widgets/SessionTabs/Archived/ArchivedConcertSessionTab.cpp
@@ -3,6 +3,7 @@
 #include "Widgets/SessionTabs/Archived/ArchivedConcertSessionTab.h"
 
 #include "ArchivedSessionHistoryController.h"
+#include "ConcertFrontendStyle.h"
 #include "ConcertLogGlobal.h"
 #include "ConcertSyncSessionDatabase.h"
 #include "HistoryEdition/ActivityNode.h"
@@ -55,6 +56,11 @@ void FArchivedConcertSessionTab::CreateDockContent(const TSharedRef<SDockTab>& I
 		);
 }
 
+const FSlateBrush* FArchivedConcertSessionTab::GetTabIconBrush() const
+{
+	return FConcertFrontendStyle::Get()->GetBrush("Concert.ArchivedSession.Icon");
+}
+
 void FArchivedConcertSessionTab::OnRequestDeleteActivities(const TSet<TSharedRef<FConcertSessionActivity>>& ActivitiesToDelete) const
 {
 	using namespace UE::ConcertSyncCore;
diff --git a/Engine/Plugins/Developer/Concert/ConcertApp/MultiUserServer/Source/MultiUserServer/Private/Widgets/SessionTabs/Archived/ArchivedConcertSessionTab.h b/Engine/Plugins/Developer/Concert/ConcertApp/MultiUserServer/Source/MultiUserServer/Private/Widgets/SessionTabs/Archived/ArchivedConcertSessionTab.h
--- a/Engine/Plugins/Developer/Concert/ConcertApp/MultiUserServer/Source/MultiUserServer/Private/Widgets/SessionTabs/Archived/ArchivedConcertSessionTab.h
+++ b/Engine/Plugins/Developer/Concert/ConcertApp/MultiUserServer/Source/MultiUserServer/Private/Widgets/SessionTabs/Archived/ArchivedConcertSessionTab.h
@@ -21,6 +21,7 @@ protected:
 
 	//~ Begin FAbstractConcertSessionTab Interface
 	virtual void CreateDockContent(const TSharedRef<SDockTab>& InDockTab) override;
+	virtual const FSlateBrush* GetTabIconBrush() const override;
 	virtual void OnOpenTab() override {}
 	//~ End FAbstractConcertSessionTab Interface
 	
diff --git a/Engine/Plugins/Developer/Concert/ConcertApp/MultiUserServer/Source/MultiUserServer/Private/Widgets/SessionTabs/ConcertSessionTabBase.cpp b/Engine/Plugins/Developer/Concert/ConcertApp/MultiUserServer/Source/MultiUserServer/Private/Widgets/SessionTabs/ConcertSessionTabBase.cpp
--- a/Engine/Plugins/Developer/Concert/ConcertApp/MultiUserServer/Source/MultiUserServer/Private/Widgets/SessionTabs/ConcertSessionTabBase.cpp
+++ b/Engine/Plugins/Developer/Concert/ConcertApp/MultiUserServer/Source/MultiUserServer/Private/Widgets/SessionTabs/ConcertSessionTabBase.cpp
@@ -7,6 +7,7 @@
 #include "Window/ConcertServerTabs.h"
 
 #include "Framework/Docking/TabManager.h"
+#include "Styling/SlateBrush.h"
 #include "Widgets/Docking/SDockTab.h"
 
 namespace UE::ConcertServerUI::Private
@@ -72,7 +73,8 @@ void FConcertSessionTabBase::EnsureInitDockTab()
 					);
 			})
 			.TabRole(MajorTab);
-		
+
+		DockTab->SetTabIcon(GetTabIconBrush());
 		CreateDockContent(DockTab.ToSharedRef());
 	}
 }
diff --git a/Engine/Plugins/Developer/Concert/ConcertApp/MultiUserServer/Source/MultiUserServer/Private/Widgets/SessionTabs/ConcertSessionTabBase.h b/Engine/Plugins/Developer/Concert/ConcertApp/MultiUserServer/Source/MultiUserServer/Private/Widgets/SessionTabs/ConcertSessionTabBase.h
--- a/Engine/Plugins/Developer/Concert/ConcertApp/MultiUserServer/Source/MultiUserServer/Private/Widgets/SessionTabs/ConcertSessionTabBase.h
+++ b/Engine/Plugins/Developer/Concert/ConcertApp/MultiUserServer/Source/MultiUserServer/Private/Widgets/SessionTabs/ConcertSessionTabBase.h
@@ -7,6 +7,7 @@
 class IConcertSyncServer;
 class SDockTab;
 class SWindow;
+struct FSlateBrush;
 
 /** Shared functionality for a tab that manages a session */
 class FConcertSessionTabBase : public TSharedFromThis<FConcertSessionTabBase>
@@ -22,6 +23,7 @@ public:
 protected:
 
 	virtual void CreateDockContent(const TSharedRef<SDockTab>& DockTab) = 0;
+	virtual const FSlateBrush* GetTabIconBrush() const = 0;
 	virtual void OnOpenTab() = 0;
 	
 	FGuid GetSessionID() const { return InspectedSessionID; }
diff --git a/Engine/Plugins/Developer/Concert/ConcertApp/MultiUserServer/Source/MultiUserServer/Private/Widgets/SessionTabs/Live/LiveConcertSessionTab.cpp b/Engine/Plugins/Developer/Concert/ConcertApp/MultiUserServer/Source/MultiUserServer/Private/Widgets/SessionTabs/Live/LiveConcertSessionTab.cpp
--- a/Engine/Plugins/Developer/Concert/ConcertApp/MultiUserServer/Source/MultiUserServer/Private/Widgets/SessionTabs/Live/LiveConcertSessionTab.cpp
+++ b/Engine/Plugins/Developer/Concert/ConcertApp/MultiUserServer/Source/MultiUserServer/Private/Widgets/SessionTabs/Live/LiveConcertSessionTab.cpp
@@ -2,11 +2,11 @@
 
 #include "LiveConcertSessionTab.h"
 
+#include "ConcertFrontendStyle.h"
 #include "IConcertSession.h"
 #include "LiveServerSessionHistoryController.h"
 #include "PackageViewer/ConcertSessionPackageViewerController.h"
 #include "SConcertLiveSessionTabView.h"
-#include "Widgets/StatusBar/SConcertStatusBar.h"
 
 #include "Widgets/Docking/SDockTab.h"
 
@@ -36,6 +36,11 @@ void FLiveConcertSessionTab::CreateDockContent(const TSharedRef<SDockTab>& InDoc
 		}));
 }
 
+const FSlateBrush* FLiveConcertSessionTab::GetTabIconBrush() const
+{
+	return FConcertFrontendStyle::Get()->GetBrush("Concert.ActiveSession.Icon");
+}
+
 void FLiveConcertSessionTab::OnOpenTab()
 {
 	SessionHistoryController->ReloadActivities();
diff --git a/Engine/Plugins/Developer/Concert/ConcertApp/MultiUserServer/Source/MultiUserServer/Private/Widgets/SessionTabs/Live/LiveConcertSessionTab.h b/Engine/Plugins/Developer/Concert/ConcertApp/MultiUserServer/Source/MultiUserServer/Private/Widgets/SessionTabs/Live/LiveConcertSessionTab.h
--- a/Engine/Plugins/Developer/Concert/ConcertApp/MultiUserServer/Source/MultiUserServer/Private/Widgets/SessionTabs/Live/LiveConcertSessionTab.h
+++ b/Engine/Plugins/Developer/Concert/ConcertApp/MultiUserServer/Source/MultiUserServer/Private/Widgets/SessionTabs/Live/LiveConcertSessionTab.h
@@ -33,6 +33,7 @@ protected:
 
 	//~ Begin FAbstractConcertSessionTab Interface
 	virtual void CreateDockContent(const TSharedRef<SDockTab>& InDockTab) override;
+	virtual const FSlateBrush* GetTabIconBrush() const override;
 	virtual void OnOpenTab() override;
 	//~ End FAbstractConcertSessionTab Interface
 	
diff --git a/Engine/Plugins/Developer/Concert/ConcertUI/ConcertSharedSlate/Source/ConcertSharedSlate/Private/ConcertFrontendStyle.cpp b/Engine/Plugins/Developer/Concert/ConcertUI/ConcertSharedSlate/Source/ConcertSharedSlate/Private/ConcertFrontendStyle.cpp
--- a/Engine/Plugins/Developer/Concert/ConcertUI/ConcertSharedSlate/Source/ConcertSharedSlate/Private/ConcertFrontendStyle.cpp
+++ b/Engine/Plugins/Developer/Concert/ConcertUI/ConcertSharedSlate/Source/ConcertSharedSlate/Private/ConcertFrontendStyle.cpp
@@ -50,6 +50,12 @@ void FConcertFrontendStyle::Initialize()
 
 	// Use this to change the opacity. Ex: In UE4, the icon looked fine at 80%, in UE5, icons looks homogeneous with others at 100%.
 	const FLinearColor IconColorAndOpacity(FLinearColor(1.f, 1.f, 1.f, 1.f));
+	const FLinearColor ActiveIconColorAndOpacity(FLinearColor(1.f, 1.f, 1.f, 1.f));
+	const FLinearColor ArchivedIconColorAndOpacity(FLinearColor(0.9f, 0.13f, 0.54f, 1.f));
+
+	// Colors
+	StyleSet->Set("Concert.ActiveSession.Color", ActiveIconColorAndOpacity);
+	StyleSet->Set("Concert.ArchivedSession.Color", ArchivedIconColorAndOpacity);
 
 	// Expandable area
 	StyleSet->Set("DetailsView.CollapsedCategory", new BOX_BRUSH( "Common/GroupBorder", FMargin(4.0f/16.0f) ) );
@@ -73,22 +79,26 @@ void FConcertFrontendStyle::Initialize()
 	StyleSet->Set("Concert.LaunchServer", new IMAGE_PLUGIN_BRUSH("Icons/icon_NewServer_32x", Icon16x16, IconColorAndOpacity));
 
 	// Multi-User Browser
-	StyleSet->Set("Concert.ArchiveSession",     new IMAGE_PLUGIN_BRUSH("Icons/icon_ArchiveSession_48x",     Icon24x24, IconColorAndOpacity));
-	StyleSet->Set("Concert.CancelAutoJoin",     new IMAGE_PLUGIN_BRUSH("Icons/icon_CancelAutoJoin_48x",     Icon24x24, IconColorAndOpacity));
-	StyleSet->Set("Concert.CloseServer",        new IMAGE_PLUGIN_BRUSH("Icons/icon_CloseServer_48x",        Icon24x24, IconColorAndOpacity));
-	StyleSet->Set("Concert.DeleteSession",      new IMAGE_PLUGIN_BRUSH("Icons/icon_DeleteSession_48x",      Icon24x24, IconColorAndOpacity));
-	StyleSet->Set("Concert.JoinDefaultSession", new IMAGE_PLUGIN_BRUSH("Icons/icon_JoinDefaultSession_48x", Icon24x24, IconColorAndOpacity));
-	StyleSet->Set("Concert.JoinSession",        new IMAGE_PLUGIN_BRUSH("Icons/icon_JoinSelectedSession_48x",Icon24x24, IconColorAndOpacity));
-	StyleSet->Set("Concert.LeaveSession",       new IMAGE_PLUGIN_BRUSH("Icons/icon_LeaveSession_48x",       Icon24x24, IconColorAndOpacity));
-	StyleSet->Set("Concert.NewServer",          new IMAGE_PLUGIN_BRUSH("Icons/icon_NewServer_48x",          Icon24x24, IconColorAndOpacity));
-	StyleSet->Set("Concert.NewSession",         new IMAGE_PLUGIN_BRUSH("Icons/icon_NewSession_48x",         Icon24x24, IconColorAndOpacity));
-	StyleSet->Set("Concert.PauseSession",       new IMAGE_PLUGIN_BRUSH("Icons/icon_PauseSession_48x",       Icon24x24, IconColorAndOpacity));
-	StyleSet->Set("Concert.RestoreSession",     new IMAGE_PLUGIN_BRUSH("Icons/icon_RestoreSession_48x",     Icon24x24, IconColorAndOpacity));
-	StyleSet->Set("Concert.ResumeSession",      new IMAGE_PLUGIN_BRUSH("Icons/icon_ResumeSession_48x",      Icon24x24, IconColorAndOpacity));
-	StyleSet->Set("Concert.Settings",           new IMAGE_PLUGIN_BRUSH("Icons/icon_Settings_48x",           Icon24x24, IconColorAndOpacity));
-	StyleSet->Set("Concert.NewServer.Small",    new IMAGE_PLUGIN_BRUSH("Icons/icon_NewServer_32x",          Icon16x16, IconColorAndOpacity));
-	StyleSet->Set("Concert.NewSession.Small",   new IMAGE_PLUGIN_BRUSH("Icons/icon_NewSession_32x",         Icon16x16, IconColorAndOpacity));
-
+	StyleSet->Set("Concert.ArchiveSession",			new IMAGE_PLUGIN_BRUSH("Icons/icon_ArchiveSession_48x",			Icon24x24, IconColorAndOpacity));
+	StyleSet->Set("Concert.CancelAutoJoin",			new IMAGE_PLUGIN_BRUSH("Icons/icon_CancelAutoJoin_48x",			Icon24x24, IconColorAndOpacity));
+	StyleSet->Set("Concert.CloseServer",				new IMAGE_PLUGIN_BRUSH("Icons/icon_CloseServer_48x",			Icon24x24, IconColorAndOpacity));
+	StyleSet->Set("Concert.DeleteSession",			new IMAGE_PLUGIN_BRUSH("Icons/icon_DeleteSession_48x",			Icon24x24, IconColorAndOpacity));
+	StyleSet->Set("Concert.JoinDefaultSession",		new IMAGE_PLUGIN_BRUSH("Icons/icon_JoinDefaultSession_48x",		Icon24x24, IconColorAndOpacity));
+	StyleSet->Set("Concert.JoinSession",				new IMAGE_PLUGIN_BRUSH("Icons/icon_JoinSelectedSession_48x",	Icon24x24, IconColorAndOpacity));
+	StyleSet->Set("Concert.LeaveSession",			new IMAGE_PLUGIN_BRUSH("Icons/icon_LeaveSession_48x",			Icon24x24, IconColorAndOpacity));
+	StyleSet->Set("Concert.NewServer",				new IMAGE_PLUGIN_BRUSH("Icons/icon_NewServer_48x",				Icon24x24, IconColorAndOpacity));
+	StyleSet->Set("Concert.NewSession",				new IMAGE_PLUGIN_BRUSH("Icons/icon_NewSession_48x",				Icon24x24, IconColorAndOpacity));
+	StyleSet->Set("Concert.PauseSession",			new IMAGE_PLUGIN_BRUSH("Icons/icon_PauseSession_48x",			Icon24x24, IconColorAndOpacity));
+	StyleSet->Set("Concert.RestoreSession",			new IMAGE_PLUGIN_BRUSH("Icons/icon_RestoreSession_48x",			Icon24x24, IconColorAndOpacity));
+	StyleSet->Set("Concert.ResumeSession",			new IMAGE_PLUGIN_BRUSH("Icons/icon_ResumeSession_48x",			Icon24x24, IconColorAndOpacity));
+	StyleSet->Set("Concert.Settings",				new IMAGE_PLUGIN_BRUSH("Icons/icon_Settings_48x",				Icon24x24, IconColorAndOpacity));
+	StyleSet->Set("Concert.NewServer.Small",			new IMAGE_PLUGIN_BRUSH("Icons/icon_NewServer_32x",				Icon16x16, IconColorAndOpacity));
+	StyleSet->Set("Concert.NewSession.Small",		new IMAGE_PLUGIN_BRUSH("Icons/icon_NewSession_32x",				Icon16x16, IconColorAndOpacity));
+	StyleSet->Set("Concert.ActiveSession.Icon",		new IMAGE_PLUGIN_BRUSH_SVG("Icons/icon_MultiUser",				Icon16x16, ActiveIconColorAndOpacity));
+	StyleSet->Set("Concert.ArchivedSession.Icon",	new IMAGE_PLUGIN_BRUSH_SVG("Icons/icon_MultiUser",				Icon16x16, ArchivedIconColorAndOpacity));
+	StyleSet->Set("Concert.SessionBrowser.FontSize", 11.f);
+	StyleSet->Set("Concert.SessionRowPadding", FMargin(0.f, 4.f));
+	
 	// Multi-user Active session
 	StyleSet->Set("Concert.JumpToLocation",     new IMAGE_PLUGIN_BRUSH("Icons/icon_PresenceLocation_32x",   Icon16x16, IconColorAndOpacity));
 	StyleSet->Set("Concert.HidePresence",       new IMAGE_PLUGIN_BRUSH("Icons/icon_PresenceEyeOff_32x",     Icon16x16, IconColorAndOpacity));
@@ -135,7 +145,7 @@ void FConcertFrontendStyle::Initialize()
 
 	// Colors
 	StyleSet->Set("Concert.DisconnectedColor", FLinearColor(0.672f, 0.672f, 0.672f));
-
+	
 	FSlateStyleRegistry::RegisterSlateStyle(*StyleSet.Get());
 };
 
diff --git a/Engine/Plugins/Developer/Concert/ConcertUI/ConcertSharedSlate/Source/ConcertSharedSlate/Private/Session/Browser/SConcertSessionBrowser.cpp b/Engine/Plugins/Developer/Concert/ConcertUI/ConcertSharedSlate/Source/ConcertSharedSlate/Private/Session/Browser/SConcertSessionBrowser.cpp
--- a/Engine/Plugins/Developer/Concert/ConcertUI/ConcertSharedSlate/Source/ConcertSharedSlate/Private/Session/Browser/SConcertSessionBrowser.cpp
+++ b/Engine/Plugins/Developer/Concert/ConcertUI/ConcertSharedSlate/Source/ConcertSharedSlate/Private/Session/Browser/SConcertSessionBrowser.cpp
@@ -6,18 +6,20 @@
 #include "ConcertHeaderRowUtils.h"
 #include "Session/Browser/ConcertBrowserUtils.h"
 #include "Session/Browser/ConcertSessionBrowserSettings.h"
-#include "Session/Browser/ConcertSessionItem.h"
+#include "Session/Browser/Items/ConcertArchivedGroupTreeItem.h"
+#include "Session/Browser/Items/ConcertSessionTreeItem.h"
 #include "Session/Browser/IConcertSessionBrowserController.h"
 #include "Session/Browser/SNewSessionRow.h"
 #include "Session/Browser/SSaveRestoreSessionRow.h"
 #include "Session/Browser/SSessionRow.h"
+#include "Session/Browser/SSessionGroupRow.h"
 
 #include "Algo/ForEach.h"
 #include "Algo/Find.h"
 #include "Framework/MultiBox/MultiBoxBuilder.h"
 #include "Misc/TextFilter.h"
+#include "Session/Browser/Items/ConcertActiveGroupTreeItem.h"
 #include "Styling/AppStyle.h"
-#include "Widgets/Images/SImage.h"
 #include "Widgets/Input/SEditableTextBox.h"
 #include "Widgets/Input/SSearchBox.h"
 #include "Widgets/Layout/SBorder.h"
@@ -61,13 +63,21 @@ void SConcertSessionBrowser::Construct(const FArguments& InArgs, TSharedRef<ICon
 
 	// Setup search filter.
 	SearchedText = InSearchText; // Reload a previous search text (in any). Useful to remember searched text between join/leave sessions, but not persistent if the tab is closed.
-	SearchTextFilter = MakeShared<TTextFilter<const FConcertSessionItem&>>(TTextFilter<const FConcertSessionItem&>::FItemToStringArray::CreateSP(this, &SConcertSessionBrowser::PopulateSearchStrings));
+	SearchTextFilter = MakeShared<TTextFilter<const FConcertSessionTreeItem&>>(TTextFilter<const FConcertSessionTreeItem&>::FItemToStringArray::CreateSP(this, &SConcertSessionBrowser::PopulateSearchStrings));
 	SearchTextFilter->OnChanged().AddSP(this, &SConcertSessionBrowser::RefreshSessionList);
 
+	auto GetSessions = [this]() { return Sessions; };
+	ActiveSessionsGroupItem = MakeShared<FConcertActiveGroupTreeItem>(GetSessions);
+	ArchivedSessionsGroupItem = MakeShared<FConcertArchivedGroupTreeItem>(GetSessions);
+	TreeItemSource = { ActiveSessionsGroupItem, ArchivedSessionsGroupItem };
+
 	ChildSlot
 	[
 		MakeBrowserContent(InArgs)
 	];
+	// By default have them expanded
+	SessionsView->SetItemExpansion(ActiveSessionsGroupItem, true);
+	SessionsView->SetItemExpansion(ArchivedSessionsGroupItem, true);
 
 	if (!SearchedText->IsEmpty())
 	{
@@ -146,21 +156,27 @@ TSharedRef<SWidget> SConcertSessionBrowser::MakeBrowserContent(const FArguments&
 		];
 }
 
+
+void SConcertSessionBrowser::OnGetChildren(TSharedPtr<FConcertTreeItem> InItem, TArray<TSharedPtr<FConcertTreeItem>>& OutChildren)
+{
+	return InItem->GetChildren(OutChildren);
+}
+
 namespace UE::SessionBrowser::Private
 {
 
 template <typename SessionType>
-constexpr FConcertSessionItem::EType GetEType()
+constexpr FConcertSessionTreeItem::EType GetEType()
 {
 	if constexpr(std::is_same<IConcertSessionBrowserController::FActiveSessionInfo,SessionType>::value)
 	{
-		return FConcertSessionItem::EType::ActiveSession;
+		return FConcertSessionTreeItem::EType::ActiveSession;
 	}
-	return FConcertSessionItem::EType::ArchivedSession;
+	return FConcertSessionTreeItem::EType::ArchivedSession;
 }
 
 template <typename SessionType>
-static FConcertSessionItem GenerateSessionItem(const SessionType& Session)
+static FConcertSessionTreeItem GenerateSessionItem(const SessionType& Session)
 {
 	FText Version = LOCTEXT("ConcertEmptyVersion", "Unset");
 	if (Session.SessionInfo.VersionInfos.Num()>0)
@@ -168,7 +184,7 @@ static FConcertSessionItem GenerateSessionItem(const SessionType& Session)
 		const FConcertSessionVersionInfo& VersionInfo = Session.SessionInfo.VersionInfos.Last();
 		Version = VersionInfo.AsText();
 	}
-	return FConcertSessionItem{
+	return FConcertSessionTreeItem(
 		GetEType<SessionType>(),
 		Session.ServerInfo.AdminEndpointId,
 		Session.SessionInfo.SessionId,
@@ -179,7 +195,7 @@ static FConcertSessionItem GenerateSessionItem(const SessionType& Session)
 		Session.ServerInfo.ServerFlags,
 		Session.SessionInfo.GetLastModified(),
 		{}
-	};
+	);
 }
 
 }
@@ -187,17 +203,17 @@ static FConcertSessionItem GenerateSessionItem(const SessionType& Session)
 void SConcertSessionBrowser::RefreshSessionList()
 {
 	// Remember the selected instances (if any).
-	TArray<TSharedPtr<FConcertSessionItem>> SelectedItems = SessionsView->GetSelectedItems();
-	TArray<TSharedPtr<FConcertSessionItem>> ReselectedItems;
-	TSharedPtr<FConcertSessionItem> NewEditableRowParent;
+	TArray<TSharedPtr<FConcertSessionTreeItem>> SelectedItems = GetSelectedItems();
+	TArray<TSharedPtr<FConcertSessionTreeItem>> ReselectedItems;
+	TSharedPtr<FConcertSessionTreeItem> NewEditableRowParent;
 
 	// Predicate returning true if the specified item should be re-selected.
-	auto IsSelected = [&SelectedItems](const FConcertSessionItem& Item)
+	auto IsSelected = [&SelectedItems](const FConcertSessionTreeItem& Item)
 	{
-		return SelectedItems.ContainsByPredicate([Item](const TSharedPtr<FConcertSessionItem>& Visited) { return *Visited == Item; });
+		return SelectedItems.ContainsByPredicate([Item](const TSharedPtr<FConcertSessionTreeItem>& Visited) { return *Visited == Item; });
 	};
 	// Matches the object instances before the update to the new instance after the update.
-	auto ReconcileObjectInstances = [&](TSharedPtr<FConcertSessionItem> NewItem)
+	auto ReconcileObjectInstances = [&](TSharedPtr<FConcertSessionTreeItem> NewItem)
 	{
 		if (IsSelected(*NewItem))
 		{
@@ -210,15 +226,15 @@ void SConcertSessionBrowser::RefreshSessionList()
 	};
 	auto ConvertSession = [&](const auto& Session)
 	{
-		FConcertSessionItem NewItem = UE::SessionBrowser::Private::GenerateSessionItem(Session);
+		FConcertSessionTreeItem NewItem = UE::SessionBrowser::Private::GenerateSessionItem(Session);
 		if (!IsFilteredOut(NewItem))
 		{
-			Sessions.Emplace(
-				MakeShared<FConcertSessionItem>(MoveTemp(NewItem)));
-			ReconcileObjectInstances(Sessions.Last());
+			const TSharedPtr<FConcertSessionTreeItem> Item = MakeShared<FConcertSessionTreeItem>(MoveTemp(NewItem));
+			AddSession(Item);
+			ReconcileObjectInstances(Item);
 		}
 	};
-	Sessions.Reset();
+	ResetSessions();
 
 	TArray<IConcertSessionBrowserController::FActiveSessionInfo> ActiveSessions = GetController()->GetActiveSessions();
 	TArray<IConcertSessionBrowserController::FArchivedSessionInfo> ArchivedSessions = GetController()->GetArchivedSessions();
@@ -229,20 +245,20 @@ void SConcertSessionBrowser::RefreshSessionList()
 	EditableSessionRowParent = NewEditableRowParent;
 	if (EditableSessionRow.IsValid())
 	{
-		if (EditableSessionRow->Type == FConcertSessionItem::EType::NewSession)
+		if (EditableSessionRow->Type == FConcertSessionTreeItem::EType::NewSession)
 		{
-			Sessions.Insert(EditableSessionRow, 0); // Always put 'new session' row at the top.
+			InsertSessionAtBeginning(EditableSessionRow); // Always put 'new session' row at the top.
 		}
 		else if (EditableSessionRowParent.IsValid())
 		{
-			Sessions.Add(EditableSessionRow); // SortSessionList() called below will ensure the correct parent/child order.
+			AddSession(EditableSessionRow); // SortSessionList() called below will ensure the correct parent/child order.
 		}
 	}
 
 	// Restore previous selection.
 	if (ReselectedItems.Num())
 	{
-		for (const TSharedPtr<FConcertSessionItem>& Item : ReselectedItems)
+		for (const TSharedPtr<FConcertSessionTreeItem>& Item : ReselectedItems)
 		{
 			SessionsView->SetItemSelection(Item, true);
 		}
@@ -259,6 +275,111 @@ TSharedPtr<IConcertSessionBrowserController> SConcertSessionBrowser::GetControll
 	return Result;
 }
 
+bool SConcertSessionBrowser::IsGroupItem(const TSharedPtr<FConcertTreeItem>& TreeItem) const
+{
+	return TreeItem == ActiveSessionsGroupItem || TreeItem == ArchivedSessionsGroupItem;
+}
+
+void SConcertSessionBrowser::ResetSessions()
+{
+	Sessions.Reset();
+}
+
+void SConcertSessionBrowser::AddSession(TSharedPtr<FConcertSessionTreeItem> Session)
+{
+	Sessions.Add(Session);
+}
+
+void SConcertSessionBrowser::InsertSessionAtBeginning(TSharedPtr<FConcertSessionTreeItem> NewSession)
+{
+	
+	Sessions.Insert(NewSession, 0);
+}
+
+void SConcertSessionBrowser::InsertSessionAfterParent(TSharedPtr<FConcertSessionTreeItem> NewSession, TSharedPtr<FConcertSessionTreeItem> Parent)
+{
+	const int32 TreeParentIndex = Sessions.IndexOfByKey(Parent);
+	if (TreeParentIndex != INDEX_NONE)
+	{
+		Sessions.Insert(NewSession, TreeParentIndex + 1);
+	}
+}
+
+SConcertSessionBrowser::EInsertPosition SConcertSessionBrowser::InsertSessionAfterParentIfAvailableOrAtBeginning(TSharedPtr<FConcertSessionTreeItem> NewSession, TSharedPtr<FConcertSessionTreeItem> Parent)
+{
+	const int32 ParentIndex = Sessions.IndexOfByKey(Parent);
+	const bool bPlaceAfterParent = ParentIndex != INDEX_NONE; 
+	Sessions.Insert(NewSession, bPlaceAfterParent ? ParentIndex + 1 : 0);
+	return bPlaceAfterParent ? EInsertPosition::AfterParent : EInsertPosition::AtBeginning;
+}
+
+void SConcertSessionBrowser::RemoveSession(TSharedPtr<FConcertSessionTreeItem> Session)
+{
+	Sessions.RemoveSingle(Session);
+}
+
+void SConcertSessionBrowser::SortSessionList()
+{
+	check(!PrimarySortedColumn.IsNone()); // Should always have a primary column. User cannot clear this one.
+
+	auto Compare = [this](const TSharedPtr<FConcertSessionTreeItem>& Lhs, const TSharedPtr<FConcertSessionTreeItem>& Rhs, const FName& ColName, EColumnSortMode::Type SortMode)
+	{
+		if (Lhs->Type == FConcertSessionTreeItem::EType::NewSession) // Always keep editable 'new session' row at the top.
+		{
+			return true;
+		}
+		if (Rhs->Type == FConcertSessionTreeItem::EType::NewSession)
+		{
+			return false;
+		}
+
+		if (ColName == ConcertBrowserUtils::SessionColName)
+		{
+			return SortMode == EColumnSortMode::Ascending ? Lhs->SessionName < Rhs->SessionName : Lhs->SessionName > Rhs->SessionName;
+		}
+		if (ColName == ConcertBrowserUtils::ServerColName)
+		{
+			return SortMode == EColumnSortMode::Ascending ? Lhs->ServerName < Rhs->ServerName : Lhs->ServerName > Rhs->ServerName;
+		}
+		if (ColName == ConcertBrowserUtils::ProjectColName)
+		{
+			return SortMode == EColumnSortMode::Ascending ? Lhs->ProjectName < Rhs->ProjectName : Lhs->ProjectName > Rhs->ProjectName;
+		}
+		if (ColName == ConcertBrowserUtils::VersionColName)
+		{
+			return SortMode == EColumnSortMode::Ascending ? Lhs->ProjectVersion < Rhs->ProjectVersion : Lhs->ProjectVersion > Rhs->ProjectVersion;
+		}
+		if (ColName == ConcertBrowserUtils::LastModifiedColName)
+		{
+			return SortMode == EColumnSortMode::Ascending ? Lhs->LastModified < Rhs->LastModified : Lhs->LastModified > Rhs->LastModified;
+		}
+		else
+		{
+			// Did you add a new column?
+			checkNoEntry();
+			return false;
+		}
+	};
+
+	Sessions.Sort([&](const TSharedPtr<FConcertSessionTreeItem>& Lhs, const TSharedPtr<FConcertSessionTreeItem>& Rhs)
+	{
+		if (Compare(Lhs, Rhs, PrimarySortedColumn, PrimarySortMode))
+		{
+			return true; // Lhs must be before Rhs based on the primary sort order.
+		}
+		if (Compare(Rhs, Lhs, PrimarySortedColumn, PrimarySortMode)) // Invert operands order (goal is to check if operands are equal or not)
+		{
+			return false; // Rhs must be before Lhs based on the primary sort.
+		}
+		else // Lhs == Rhs on the primary column, need to order accoding the secondary column if one is set.
+		{
+			return SecondarySortedColumn.IsNone() ? false : Compare(Lhs, Rhs, SecondarySortedColumn, SecondarySortMode);
+		}
+	});
+
+	EnsureEditableParentChildOrder();
+}
+
 void SConcertSessionBrowser::OnSearchTextChanged(const FText& InFilterText)
 {
 	SearchTextFilter->SetRawFilterText(InFilterText);
@@ -276,18 +397,18 @@ void SConcertSessionBrowser::OnSearchTextCommitted(const FText& InFilterText, ET
 	}
 }
 
-void SConcertSessionBrowser::PopulateSearchStrings(const FConcertSessionItem& Item, TArray<FString>& OutSearchStrings) const
+void SConcertSessionBrowser::PopulateSearchStrings(const FConcertSessionTreeItem& Item, TArray<FString>& OutSearchStrings) const
 {
 	OutSearchStrings.Add(Item.ServerName);
 	OutSearchStrings.Add(Item.SessionName);
 }
 
-bool SConcertSessionBrowser::IsFilteredOut(const FConcertSessionItem& Item) const
+bool SConcertSessionBrowser::IsFilteredOut(const FConcertSessionTreeItem& Item) const
 {
 	bool bIsDefaultServer = LastDefaultServerUrl.IsEmpty() || Item.ServerName == LastDefaultServerUrl;
 
-	return (!PersistentSettings->bShowActiveSessions && (Item.Type == FConcertSessionItem::EType::ActiveSession || Item.Type == FConcertSessionItem::EType::SaveSession)) ||
-		   (!PersistentSettings->bShowArchivedSessions && (Item.Type == FConcertSessionItem::EType::ArchivedSession || Item.Type == FConcertSessionItem::EType::RestoreSession)) ||
+	return (!PersistentSettings->bShowActiveSessions && (Item.Type == FConcertSessionTreeItem::EType::ActiveSession || Item.Type == FConcertSessionTreeItem::EType::SaveSession)) ||
+		   (!PersistentSettings->bShowArchivedSessions && (Item.Type == FConcertSessionTreeItem::EType::ArchivedSession || Item.Type == FConcertSessionTreeItem::EType::RestoreSession)) ||
 		   (PersistentSettings->bShowDefaultServerSessionsOnly && !bIsDefaultServer) ||
 		   !SearchTextFilter->PassesFilter(Item);
 }
@@ -375,14 +496,15 @@ TSharedRef<SWidget> SConcertSessionBrowser::MakeSessionTableView(const FArgument
 {
 	using namespace UE::ConcertSharedSlate;
 
-	PrimarySortedColumn = ConcertBrowserUtils::IconColName;
+	PrimarySortedColumn = ConcertBrowserUtils::SessionColName;
 	PrimarySortMode = EColumnSortMode::Ascending;
-	SecondarySortedColumn = ConcertBrowserUtils::SessionColName;
+	SecondarySortedColumn = FName();
 	SecondarySortMode = EColumnSortMode::Ascending;
 
-	SessionsView = SNew(SListView<TSharedPtr<FConcertSessionItem>>)
-		.ListItemsSource(&Sessions)
+	SessionsView = SNew(STreeView<TSharedPtr<FConcertTreeItem>>)
+		.TreeItemsSource(&TreeItemSource)
 		.OnGenerateRow(this, &SConcertSessionBrowser::OnGenerateSessionRowWidget)
+		.OnGetChildren(this, &SConcertSessionBrowser::OnGetChildren)
 		.SelectionMode(ESelectionMode::Multi)
 		.OnSelectionChanged(this, &SConcertSessionBrowser::OnSessionSelectionChanged)
 		.OnContextMenuOpening(this, &SConcertSessionBrowser::MakeTableContextualMenu)
@@ -397,16 +519,9 @@ TSharedRef<SWidget> SConcertSessionBrowser::MakeSessionTableView(const FArgument
 				}
 			})
 
-			+SHeaderRow::Column(ConcertBrowserUtils::IconColName)
-			.DefaultLabel(FText::GetEmpty())
-			.SortPriority(this, &SConcertSessionBrowser::GetColumnSortPriority, ConcertBrowserUtils::IconColName)
-			.SortMode(this, &SConcertSessionBrowser::GetColumnSortMode, ConcertBrowserUtils::IconColName)
-			.OnSort(this, &SConcertSessionBrowser::OnColumnSortModeChanged)
-			.FixedWidth(20)
-			.ShouldGenerateWidget(true)
-
 			+SHeaderRow::Column(ConcertBrowserUtils::SessionColName)
 			.DefaultLabel(LOCTEXT("SessioName", "Session"))
+			.FillWidth(2.f)
 			.SortPriority(this, &SConcertSessionBrowser::GetColumnSortPriority, ConcertBrowserUtils::SessionColName)
 			.SortMode(this, &SConcertSessionBrowser::GetColumnSortMode, ConcertBrowserUtils::SessionColName)
 			.OnSort(this, &SConcertSessionBrowser::OnColumnSortModeChanged)
@@ -414,24 +529,28 @@ TSharedRef<SWidget> SConcertSessionBrowser::MakeSessionTableView(const FArgument
 
 			+SHeaderRow::Column(ConcertBrowserUtils::ServerColName)
 			.DefaultLabel(LOCTEXT("Server", "Server"))
+			.FillWidth(1.f)
 			.SortPriority(this, &SConcertSessionBrowser::GetColumnSortPriority, ConcertBrowserUtils::ServerColName)
 			.SortMode(this, &SConcertSessionBrowser::GetColumnSortMode, ConcertBrowserUtils::ServerColName)
 			.OnSort(this, &SConcertSessionBrowser::OnColumnSortModeChanged)
 			
 			+SHeaderRow::Column(ConcertBrowserUtils::ProjectColName)
 			.DefaultLabel(LOCTEXT("ProjectCol", "Project"))
+			.FillWidth(1.f)
 			.SortPriority(this, &SConcertSessionBrowser::GetColumnSortPriority, ConcertBrowserUtils::ProjectColName)
 			.SortMode(this, &SConcertSessionBrowser::GetColumnSortMode, ConcertBrowserUtils::ProjectColName)
 			.OnSort(this, &SConcertSessionBrowser::OnColumnSortModeChanged)
 			
 			+SHeaderRow::Column(ConcertBrowserUtils::VersionColName)
 			.DefaultLabel(LOCTEXT("Version", "Version"))
+			.FillWidth(1.f)
 			.SortPriority(this, &SConcertSessionBrowser::GetColumnSortPriority, ConcertBrowserUtils::VersionColName)
 			.SortMode(this, &SConcertSessionBrowser::GetColumnSortMode, ConcertBrowserUtils::VersionColName)
 			.OnSort(this, &SConcertSessionBrowser::OnColumnSortModeChanged)
 			
 			+SHeaderRow::Column(ConcertBrowserUtils::LastModifiedColName)
 			.DefaultLabel(LOCTEXT("LastModified", "Last Modified"))
+			.FillWidth(1.f)
 			.DefaultTooltip(LOCTEXT("LastModifiedTooltip", "The local time the session's directory was last modified"))
 			.SortPriority(this, &SConcertSessionBrowser::GetColumnSortPriority, ConcertBrowserUtils::LastModifiedColName)
 			.SortMode(this, &SConcertSessionBrowser::GetColumnSortMode, ConcertBrowserUtils::LastModifiedColName)
@@ -493,121 +612,58 @@ void SConcertSessionBrowser::OnColumnSortModeChanged(const EColumnSortPriority::
 	SessionsView->RequestListRefresh();
 }
 
-void SConcertSessionBrowser::SortSessionList()
-{
-	check(!PrimarySortedColumn.IsNone()); // Should always have a primary column. User cannot clear this one.
-
-	auto Compare = [this](const TSharedPtr<FConcertSessionItem>& Lhs, const TSharedPtr<FConcertSessionItem>& Rhs, const FName& ColName, EColumnSortMode::Type SortMode)
-	{
-		if (Lhs->Type == FConcertSessionItem::EType::NewSession) // Always keep editable 'new session' row at the top.
-		{
-			return true;
-		}
-		if (Rhs->Type == FConcertSessionItem::EType::NewSession)
-		{
-			return false;
-		}
-
-		if (ColName == ConcertBrowserUtils::IconColName)
-		{
-			return SortMode == EColumnSortMode::Ascending ? Lhs->Type < Rhs->Type : Lhs->Type > Rhs->Type;
-		}
-		if (ColName == ConcertBrowserUtils::SessionColName)
-		{
-			return SortMode == EColumnSortMode::Ascending ? Lhs->SessionName < Rhs->SessionName : Lhs->SessionName > Rhs->SessionName;
-		}
-		if (ColName == ConcertBrowserUtils::ServerColName)
-		{
-			return SortMode == EColumnSortMode::Ascending ? Lhs->ServerName < Rhs->ServerName : Lhs->ServerName > Rhs->ServerName;
-		}
-		if (ColName == ConcertBrowserUtils::ProjectColName)
-		{
-			return SortMode == EColumnSortMode::Ascending ? Lhs->ProjectName < Rhs->ProjectName : Lhs->ProjectName > Rhs->ProjectName;
-		}
-		if (ColName == ConcertBrowserUtils::VersionColName)
-		{
-			return SortMode == EColumnSortMode::Ascending ? Lhs->ProjectVersion < Rhs->ProjectVersion : Lhs->ProjectVersion > Rhs->ProjectVersion;
-		}
-		if (ColName == ConcertBrowserUtils::LastModifiedColName)
-		{
-			return SortMode == EColumnSortMode::Ascending ? Lhs->LastModified < Rhs->LastModified : Lhs->LastModified > Rhs->LastModified;
-		}
-		else
-		{
-			// Did you add a new column?
-			checkNoEntry();
-			return false;
-		}
-	};
-
-	Sessions.Sort([&](const TSharedPtr<FConcertSessionItem>& Lhs, const TSharedPtr<FConcertSessionItem>& Rhs)
-	{
-		if (Compare(Lhs, Rhs, PrimarySortedColumn, PrimarySortMode))
-		{
-			return true; // Lhs must be before Rhs based on the primary sort order.
-		}
-		if (Compare(Rhs, Lhs, PrimarySortedColumn, PrimarySortMode)) // Invert operands order (goal is to check if operands are equal or not)
-		{
-			return false; // Rhs must be before Lhs based on the primary sort.
-		}
-		else // Lhs == Rhs on the primary column, need to order accoding the secondary column if one is set.
-		{
-			return SecondarySortedColumn.IsNone() ? false : Compare(Lhs, Rhs, SecondarySortedColumn, SecondarySortMode);
-		}
-	});
-
-	EnsureEditableParentChildOrder();
-}
-
 void SConcertSessionBrowser::EnsureEditableParentChildOrder()
 {
 	// This is for Archiving or Restoring a session. We keep the editable row below the session to archive or restore and visually link them with small wires in UI.
 	if (EditableSessionRowParent.IsValid())
 	{
 		check(EditableSessionRow.IsValid());
-		Sessions.Remove(EditableSessionRow);
-
-		int32 ParentIndex = Sessions.IndexOfByKey(EditableSessionRowParent);
-		if (ParentIndex != INDEX_NONE)
-		{
-			Sessions.Insert(EditableSessionRow, ParentIndex + 1); // Insert the 'child' below its parent.
-		}
+		RemoveSession(EditableSessionRow);
+		
+		InsertSessionAfterParent(EditableSessionRow, EditableSessionRowParent); // Insert the 'child' below its parent.
 	}
 }
 
-TSharedRef<ITableRow> SConcertSessionBrowser::OnGenerateSessionRowWidget(TSharedPtr<FConcertSessionItem> Item, const TSharedRef<STableViewBase>& OwnerTable)
+TSharedRef<ITableRow> SConcertSessionBrowser::OnGenerateSessionRowWidget(TSharedPtr<FConcertTreeItem> Item, const TSharedRef<STableViewBase>& OwnerTable)
 {
-	switch (Item->Type)
+	if (IsGroupItem(Item))
+	{
+		return MakeGroupRowWidget(Item, OwnerTable);
+	}
+
+	const TSharedPtr<FConcertSessionTreeItem> SessionItem = StaticCastSharedPtr<FConcertSessionTreeItem>(Item);
+	switch (SessionItem->Type)
 	{
-		case FConcertSessionItem::EType::ActiveSession:
-			return MakeActiveSessionRowWidget(Item, OwnerTable);
+		case FConcertSessionTreeItem::EType::ActiveSession:
+			return MakeActiveSessionRowWidget(SessionItem, OwnerTable);
 
-		case FConcertSessionItem::EType::ArchivedSession:
-			return MakeArchivedSessionRowWidget(Item, OwnerTable);
+		case FConcertSessionTreeItem::EType::ArchivedSession:
+			return MakeArchivedSessionRowWidget(SessionItem, OwnerTable);
 
-		case FConcertSessionItem::EType::NewSession:
-			return MakeNewSessionRowWidget(Item, OwnerTable);
+		case FConcertSessionTreeItem::EType::NewSession:
+			return MakeNewSessionRowWidget(SessionItem, OwnerTable);
 
-		case FConcertSessionItem::EType::RestoreSession:
-			return MakeRestoreSessionRowWidget(Item, OwnerTable);
+		case FConcertSessionTreeItem::EType::RestoreSession:
+			return MakeRestoreSessionRowWidget(SessionItem, OwnerTable);
 
 		default:
-			check(Item->Type == FConcertSessionItem::EType::SaveSession);
-			return MakeSaveSessionRowWidget(Item, OwnerTable);
+			check(SessionItem->Type == FConcertSessionTreeItem::EType::SaveSession);
+			return MakeSaveSessionRowWidget(SessionItem, OwnerTable);
 	}
 }
 
-TSharedRef<ITableRow> SConcertSessionBrowser::MakeActiveSessionRowWidget(const TSharedPtr<FConcertSessionItem>& ActiveItem, const TSharedRef<STableViewBase>& OwnerTable)
+TSharedRef<ITableRow> SConcertSessionBrowser::MakeActiveSessionRowWidget(const TSharedPtr<FConcertSessionTreeItem>& ActiveItem, const TSharedRef<STableViewBase>& OwnerTable)
 {
 	const TOptional<FConcertSessionInfo> SessionInfo = GetController()->GetActiveSessionInfo(ActiveItem->ServerAdminEndpointId, ActiveItem->SessionId);
 
 	// Add an 'Active Session' row. Clicking the row icon joins the session.
 	return SNew(SSessionRow, ActiveItem, OwnerTable)
-		.OnDoubleClickFunc([this](const TSharedPtr<FConcertSessionItem>& Item) { OnLiveSessionDoubleClicked.ExecuteIfBound(Item); })
-		.OnRenameFunc([this](const TSharedPtr<FConcertSessionItem>& Item, const FString& NewName) { RequestRenameSession(Item, NewName); })
-		.IsDefaultSession([this](TSharedPtr<FConcertSessionItem> ItemPin)
+		.Padding(FConcertFrontendStyle::Get()->GetMargin("Concert.SessionRowPadding"))
+		.OnDoubleClickFunc([this](const TSharedPtr<FConcertSessionTreeItem>& Item) { OnLiveSessionDoubleClicked.ExecuteIfBound(Item); })
+		.OnRenameFunc([this](const TSharedPtr<FConcertSessionTreeItem>& Item, const FString& NewName) { RequestRenameSession(Item, NewName); })
+		.IsDefaultSession([this](TSharedPtr<FConcertSessionTreeItem> ItemPin)
 		{
-			return ItemPin->Type == FConcertSessionItem::EType::ActiveSession
+			return ItemPin->Type == FConcertSessionTreeItem::EType::ActiveSession
 				&& DefaultSessionName.IsSet() && ItemPin->SessionName == DefaultSessionName.Get()
 				&& DefaultServerUrl.IsSet() && ItemPin->ServerName == DefaultServerUrl.Get();
 		})
@@ -616,17 +672,18 @@ TSharedRef<ITableRow> SConcertSessionBrowser::MakeActiveSessionRowWidget(const T
 		.IsSelected_Lambda([this, ActiveItem]() { return SessionsView->GetSelectedItems().Num() == 1 && SessionsView->GetSelectedItems()[0] == ActiveItem; });
 }
 
-TSharedRef<ITableRow> SConcertSessionBrowser::MakeArchivedSessionRowWidget(const TSharedPtr<FConcertSessionItem>& ArchivedItem, const TSharedRef<STableViewBase>& OwnerTable)
+TSharedRef<ITableRow> SConcertSessionBrowser::MakeArchivedSessionRowWidget(const TSharedPtr<FConcertSessionTreeItem>& ArchivedItem, const TSharedRef<STableViewBase>& OwnerTable)
 {
 	const TOptional<FConcertSessionInfo> SessionInfo = GetController()->GetArchivedSessionInfo(ArchivedItem->ServerAdminEndpointId, ArchivedItem->SessionId);
 
 	// Add an 'Archived Session' row. Clicking the row icon adds a 'Restore as' row to the table.
 	return SNew(SSessionRow, ArchivedItem, OwnerTable)
-		.OnDoubleClickFunc([this](const TSharedPtr<FConcertSessionItem>& Item) { OnArchivedSessionDoubleClicked.Execute(Item); })
-		.OnRenameFunc([this](const TSharedPtr<FConcertSessionItem>& Item, const FString& NewName) { RequestRenameSession(Item, NewName); })
-		.IsDefaultSession([this](TSharedPtr<FConcertSessionItem> ItemPin)
+		.Padding(FConcertFrontendStyle::Get()->GetMargin("Concert.SessionRowPadding"))
+		.OnDoubleClickFunc([this](const TSharedPtr<FConcertSessionTreeItem>& Item) { OnArchivedSessionDoubleClicked.Execute(Item); })
+		.OnRenameFunc([this](const TSharedPtr<FConcertSessionTreeItem>& Item, const FString& NewName) { RequestRenameSession(Item, NewName); })
+		.IsDefaultSession([this](TSharedPtr<FConcertSessionTreeItem> ItemPin)
 		{
-			return ItemPin->Type == FConcertSessionItem::EType::ActiveSession
+			return ItemPin->Type == FConcertSessionTreeItem::EType::ActiveSession
 				&& DefaultSessionName.IsSet() && ItemPin->SessionName == DefaultSessionName.Get()
 				&& DefaultServerUrl.IsSet() && ItemPin->ServerName == DefaultServerUrl.Get();
 		})
@@ -635,66 +692,75 @@ TSharedRef<ITableRow> SConcertSessionBrowser::MakeArchivedSessionRowWidget(const
 		.IsSelected_Lambda([this, ArchivedItem]() { return SessionsView->GetSelectedItems().Num() == 1 && SessionsView->GetSelectedItems()[0] == ArchivedItem; });
 }
 
-TSharedRef<ITableRow> SConcertSessionBrowser::MakeNewSessionRowWidget(const TSharedPtr<FConcertSessionItem>& NewItem, const TSharedRef<STableViewBase>& OwnerTable)
+TSharedRef<ITableRow> SConcertSessionBrowser::MakeNewSessionRowWidget(const TSharedPtr<FConcertSessionTreeItem>& NewItem, const TSharedRef<STableViewBase>& OwnerTable)
 {
 	// Add an editable 'New Session' row in the table to let user pick a name and a server.
 	TAttribute<FString> DefaultServerUrlArg = DefaultServerUrl.IsSet()
 		? TAttribute<FString>::Create([this](){ return DefaultServerUrl.Get(); }) : TAttribute<FString>();
 	return SNew(SNewSessionRow, NewItem, OwnerTable)
+		.Padding(FConcertFrontendStyle::Get()->GetMargin("Concert.SessionRowPadding"))
 		.GetServerFunc([this]() { return GetController()->GetServers(); }) // Let the row pull the servers for the combo box.
-		.OnAcceptFunc([this](const TSharedPtr<FConcertSessionItem>& Item) { RequestCreateSession(Item); }) // Accepting creates the session.
-		.OnDeclineFunc([this](const TSharedPtr<FConcertSessionItem>& Item) { RemoveSessionRow(Item); })  // Declining removes the editable 'new session' row from the view.
+		.OnAcceptFunc([this](const TSharedPtr<FConcertSessionTreeItem>& Item) { RequestCreateSession(Item); }) // Accepting creates the session.
+		.OnDeclineFunc([this](const TSharedPtr<FConcertSessionTreeItem>& Item) { RemoveSessionRow(Item); })  // Declining removes the editable 'new session' row from the view.
 		.HighlightText(this, &SConcertSessionBrowser::HighlightSearchText)
 		.DefaultServerURL(DefaultServerUrlArg);
 }
 
-TSharedRef<ITableRow> SConcertSessionBrowser::MakeSaveSessionRowWidget(const TSharedPtr<FConcertSessionItem>& SaveItem, const TSharedRef<STableViewBase>& OwnerTable)
+TSharedRef<ITableRow> SConcertSessionBrowser::MakeSaveSessionRowWidget(const TSharedPtr<FConcertSessionTreeItem>& SaveItem, const TSharedRef<STableViewBase>& OwnerTable)
 {
 	// Add an editable 'Save Session' row in the table to let the user enter an archive name.
 	return SNew(SSaveRestoreSessionRow, SaveItem, OwnerTable)
-		.OnAcceptFunc([this](const TSharedPtr<FConcertSessionItem>& Item, const FString& ArchiveName) { RequestArchiveSession(Item, ArchiveName); }) // Accepting archive the session.
-		.OnDeclineFunc([this](const TSharedPtr<FConcertSessionItem>& Item) { RemoveSessionRow(Item); }) // Declining removes the editable 'save session as' row from the view.
+		.Padding(FConcertFrontendStyle::Get()->GetMargin("Concert.SessionRowPadding"))
+		.OnAcceptFunc([this](const TSharedPtr<FConcertSessionTreeItem>& Item, const FString& ArchiveName) { RequestArchiveSession(Item, ArchiveName); }) // Accepting archive the session.
+		.OnDeclineFunc([this](const TSharedPtr<FConcertSessionTreeItem>& Item) { RemoveSessionRow(Item); }) // Declining removes the editable 'save session as' row from the view.
 		.HighlightText(this, &SConcertSessionBrowser::HighlightSearchText);
 }
 
-TSharedRef<ITableRow> SConcertSessionBrowser::MakeRestoreSessionRowWidget(const TSharedPtr<FConcertSessionItem>& RestoreItem, const TSharedRef<STableViewBase>& OwnerTable)
+TSharedRef<ITableRow> SConcertSessionBrowser::MakeGroupRowWidget(const TSharedPtr<FConcertTreeItem>& Item, const TSharedRef<STableViewBase>& OwnerTable)
+{
+	const bool bIsActive = Item == ActiveSessionsGroupItem;
+	return SNew(SSessionGroupRow, OwnerTable)
+		.Padding(FConcertFrontendStyle::Get()->GetMargin("Concert.SessionRowPadding"))
+		.GroupType(bIsActive ? SSessionGroupRow::EGroupType::Active : SSessionGroupRow::EGroupType::Archived);
+}
+
+TSharedRef<ITableRow> SConcertSessionBrowser::MakeRestoreSessionRowWidget(const TSharedPtr<FConcertSessionTreeItem>& RestoreItem, const TSharedRef<STableViewBase>& OwnerTable)
 {
 	// Add an editable 'Restore Session' row in the table to let the user enter a session name.
 	return SNew(SSaveRestoreSessionRow, RestoreItem, OwnerTable)
-		.OnAcceptFunc([this](const TSharedPtr<FConcertSessionItem>& Item, const FString& SessionName) { RequestRestoreSession(Item, SessionName); }) // Accepting restores the session.
-		.OnDeclineFunc([this](const TSharedPtr<FConcertSessionItem>& Item) { RemoveSessionRow(Item); }) // Declining removes the editable 'restore session as' row from the view.
+		.Padding(FConcertFrontendStyle::Get()->GetMargin("Concert.SessionRowPadding"))
+		.OnAcceptFunc([this](const TSharedPtr<FConcertSessionTreeItem>& Item, const FString& SessionName) { RequestRestoreSession(Item, SessionName); }) // Accepting restores the session.
+		.OnDeclineFunc([this](const TSharedPtr<FConcertSessionTreeItem>& Item) { RemoveSessionRow(Item); }) // Declining removes the editable 'restore session as' row from the view.
 		.HighlightText(this, &SConcertSessionBrowser::HighlightSearchText);
 }
 
 void SConcertSessionBrowser::InsertNewSessionEditableRowInternal()
 {
 	// Insert a 'new session' editable row.
-	FConcertSessionItem Item{FConcertSessionItem::EType::NewSession};
-	InsertEditableSessionRow(MakeShared<FConcertSessionItem>(MoveTemp(Item)), nullptr);
+	FConcertSessionTreeItem Item;
+	Item.Type = FConcertSessionTreeItem::EType::NewSession;
+	InsertEditableSessionRow(MakeShared<FConcertSessionTreeItem>(MoveTemp(Item)), nullptr);
 }
 
-PRAGMA_DISABLE_OPTIMIZATION
-void SConcertSessionBrowser::InsertRestoreSessionAsEditableRowInternal(const TSharedPtr<FConcertSessionItem>& ArchivedItem)
+void SConcertSessionBrowser::InsertRestoreSessionAsEditableRowInternal(const TSharedPtr<FConcertSessionTreeItem>& ArchivedItem)
 {
 	// Insert the 'restore session as ' editable row just below the 'archived' item to restore.
-	InsertEditableSessionRow(MakeShared<FConcertSessionItem>(ArchivedItem->MakeCopyAsType(FConcertSessionItem::EType::RestoreSession)), ArchivedItem);
+	InsertEditableSessionRow(MakeShared<FConcertSessionTreeItem>(ArchivedItem->MakeCopyAsType(FConcertSessionTreeItem::EType::RestoreSession)), ArchivedItem);
 }
-	PRAGMA_ENABLE_OPTIMIZATION
 
-void SConcertSessionBrowser::InsertArchiveSessionAsEditableRow(const TSharedPtr<FConcertSessionItem>& LiveItem)
+void SConcertSessionBrowser::InsertArchiveSessionAsEditableRow(const TSharedPtr<FConcertSessionTreeItem>& LiveItem)
 {
 	// Insert the 'save session as' editable row just below the 'active' item to save.
-	InsertEditableSessionRow(MakeShared<FConcertSessionItem>(LiveItem->MakeCopyAsType(FConcertSessionItem::EType::SaveSession)), LiveItem);
+	InsertEditableSessionRow(MakeShared<FConcertSessionTreeItem>(LiveItem->MakeCopyAsType(FConcertSessionTreeItem::EType::SaveSession)), LiveItem);
 }
 
-void SConcertSessionBrowser::InsertEditableSessionRow(TSharedPtr<FConcertSessionItem> EditableItem, TSharedPtr<FConcertSessionItem> ParentItem)
+void SConcertSessionBrowser::InsertEditableSessionRow(TSharedPtr<FConcertSessionTreeItem> EditableItem, TSharedPtr<FConcertSessionTreeItem> ParentItem)
 {
 	// Insert the new row below its parent (if any).
-	int32 ParentIndex = Sessions.IndexOfByKey(ParentItem);
-	Sessions.Insert(EditableItem, ParentIndex != INDEX_NONE ? ParentIndex + 1 : 0);
+	const EInsertPosition InsertPosition = InsertSessionAfterParentIfAvailableOrAtBeginning(EditableItem, ParentItem);
 
 	// Ensure there is only one editable row at the time, removing the row being edited (if any).
-	Sessions.Remove(EditableSessionRow);
+	RemoveSession(EditableSessionRow);
 	EditableSessionRow = EditableItem;
 	EditableSessionRowParent = ParentItem;
 
@@ -705,7 +771,7 @@ void SConcertSessionBrowser::InsertEditableSessionRow(TSharedPtr<FConcertSession
 	// NOTE: Ideally, I would only use RequestScrollIntoView() to scroll the new item into view, but it did not work. If an item was added into an hidden part,
 	//       it was not always scrolled correctly into view. RequestNavigateToItem() worked much better, except when inserting the very first row in the list, in
 	//       such case calling the function would give the focus to the list view (showing a white dashed line around it).
-	if (ParentIndex == INDEX_NONE)
+	if (InsertPosition == EInsertPosition::AtBeginning)
 	{
 		SessionsView->ScrollToTop(); // Item is inserted at 0. (New session)
 	}
@@ -715,14 +781,14 @@ void SConcertSessionBrowser::InsertEditableSessionRow(TSharedPtr<FConcertSession
 	}
 }
 
-void SConcertSessionBrowser::RemoveSessionRow(const TSharedPtr<FConcertSessionItem>& Item)
+void SConcertSessionBrowser::RemoveSessionRow(const TSharedPtr<FConcertSessionTreeItem>& Item)
 {
-	Sessions.Remove(Item);
+	RemoveSession(Item);
 
 	// Don't keep the editable row if its 'parent' is removed. (if the session to restore or archive gets deleted in the meantime)
 	if (Item == EditableSessionRowParent)
 	{
-		Sessions.Remove(EditableSessionRow);
+		RemoveSession(EditableSessionRow);
 		EditableSessionRow.Reset();
 	}
 
@@ -746,9 +812,9 @@ TSharedRef<SWidget> SConcertSessionBrowser::MakeSessionTableFooter()
 			SNew(STextBlock).Text_Lambda([this]()
 			{
 				// Don't count the 'New Session', 'Restore Session' and 'Archive Session' editable row, they are transient rows used for inline input only.
-				int32 DisplayedSessionNum = Sessions.Num() - (EditableSessionRow.IsValid() ? 1 : 0);
-				int32 AvailableSessionNum = GetController()->GetActiveSessions().Num() + GetController()->GetArchivedSessions().Num();
-				int32 ServerNum = GetController()->GetServers().Num();
+				const int32 DisplayedSessionNum = Sessions.Num() - (EditableSessionRow.IsValid() ? 1 : 0);
+				const int32 AvailableSessionNum = GetController()->GetActiveSessions().Num() + GetController()->GetArchivedSessions().Num();
+				const int32 ServerNum = GetController()->GetServers().Num();
 
 				// If all discovered session are displayed (none excluded by a filter).
 				if (DisplayedSessionNum == AvailableSessionNum)
@@ -821,25 +887,30 @@ void SConcertSessionBrowser::OnFilterMenuChecked(const FName MenuName)
 
 TSharedPtr<SWidget> SConcertSessionBrowser::MakeTableContextualMenu()
 {
-	TArray<TSharedPtr<FConcertSessionItem>> SelectedItems = SessionsView->GetSelectedItems();
-	if (SelectedItems.Num() == 0 || (SelectedItems[0]->Type != FConcertSessionItem::EType::ActiveSession && SelectedItems[0]->Type != FConcertSessionItem::EType::ArchivedSession))
+	if (SessionsView->GetSelectedItems().Num() == 0)
 	{
 		// If clicking the header row, show the options for hiding them
 		return UE::ConcertSharedSlate::MakeTableContextMenu(SessionHeaderRow.ToSharedRef(), {}, true);
 	}
 
-	TSharedPtr<FConcertSessionItem> Item = SelectedItems[0];
+	const TArray<TSharedPtr<FConcertSessionTreeItem>> SelectedItems = GetSelectedItems();
+	if (SelectedItems.Num() == 0)
+	{
+		// If you click the groups then the context menu should not contain anything
+		return nullptr;
+	}
+	TSharedPtr<FConcertSessionTreeItem> Item = SelectedItems[0];
 
 	TSharedRef<FExtender> Extender = MakeShared<FExtender>();
 	ExtendSessionContextMenu.ExecuteIfBound(Item, Extender.Get());
 	FMenuBuilder MenuBuilder(/*bInShouldCloseWindowAfterMenuSelection=*/true, nullptr, Extender);
 
 	// Section title.
-	MenuBuilder.BeginSection(SessionContextMenuExtensionHooks::ManageSession, Item->Type == FConcertSessionItem::EType::ActiveSession ?
+	MenuBuilder.BeginSection(SessionContextMenuExtensionHooks::ManageSession, Item->Type == FConcertSessionTreeItem::EType::ActiveSession ?
 		LOCTEXT("ActiveSessionSection", "Active Session"):
 		LOCTEXT("ArchivedSessionSection", "Archived Session"));
 
-	if (Item->Type == FConcertSessionItem::EType::ActiveSession)
+	if (Item->Type == FConcertSessionTreeItem::EType::ActiveSession)
 	{
 		MenuBuilder.AddMenuEntry(
 			LOCTEXT("CtxMenuArchive", "Archive"),
@@ -955,7 +1026,15 @@ TSharedRef<SWidget> SConcertSessionBrowser::MakeViewOptionsComboButtonMenu()
 	return MenuBuilder.MakeWidget();
 }
 
-void SConcertSessionBrowser::OnSessionSelectionChanged(TSharedPtr<FConcertSessionItem> SelectedSession, ESelectInfo::Type SelectInfo)
+void SConcertSessionBrowser::OnSessionSelectionChanged(TSharedPtr<FConcertTreeItem> SelectedSession, ESelectInfo::Type SelectInfo)
+{
+	if (!IsGroupItem(SelectedSession))
+	{
+		OnSessionSelectionChanged(StaticCastSharedPtr<FConcertSessionTreeItem>(SelectedSession), SelectInfo);
+	}
+}
+
+void SConcertSessionBrowser::OnSessionSelectionChanged(TSharedPtr<FConcertSessionTreeItem> SelectedSession, ESelectInfo::Type SelectInfo)
 {
 	// Cancel editing the row to create, archive or restore a session (if any), unless the row was selected in code.
 	if (EditableSessionRow.IsValid() && SelectInfo != ESelectInfo::Direct)
@@ -978,38 +1057,38 @@ bool SConcertSessionBrowser::IsNewButtonEnabledInternal() const
 
 bool SConcertSessionBrowser::IsRestoreButtonEnabledInternal() const
 {
-	TArray<TSharedPtr<FConcertSessionItem>> SelectedItems = SessionsView->GetSelectedItems();
-	return SelectedItems.Num() == 1 && SelectedItems[0]->Type == FConcertSessionItem::EType::ArchivedSession;
+	TArray<TSharedPtr<FConcertSessionTreeItem>> SelectedItems = GetSelectedItems();
+	return SelectedItems.Num() == 1 && SelectedItems[0]->Type == FConcertSessionTreeItem::EType::ArchivedSession;
 }
 
 bool SConcertSessionBrowser::IsArchiveButtonEnabledInternal() const
 {
-	TArray<TSharedPtr<FConcertSessionItem>> SelectedItems = SessionsView->GetSelectedItems();
-	return SelectedItems.Num() == 1 && SelectedItems[0]->Type == FConcertSessionItem::EType::ActiveSession;
+	TArray<TSharedPtr<FConcertSessionTreeItem>> SelectedItems = GetSelectedItems();
+	return SelectedItems.Num() == 1 && SelectedItems[0]->Type == FConcertSessionTreeItem::EType::ActiveSession;
 }
 
 bool SConcertSessionBrowser::IsRenameButtonEnabledInternal() const
 {
-	TArray<TSharedPtr<FConcertSessionItem>> SelectedItems = SessionsView->GetSelectedItems();
+	TArray<TSharedPtr<FConcertSessionTreeItem>> SelectedItems = GetSelectedItems();
 	if (SelectedItems.Num() != 1)
 	{
 		return false;
 	}
 
-	return (SelectedItems[0]->Type == FConcertSessionItem::EType::ActiveSession && GetController()->CanRenameActiveSession(SelectedItems[0]->ServerAdminEndpointId, SelectedItems[0]->SessionId)) ||
-	       (SelectedItems[0]->Type == FConcertSessionItem::EType::ArchivedSession && GetController()->CanRenameArchivedSession(SelectedItems[0]->ServerAdminEndpointId, SelectedItems[0]->SessionId));
+	return (SelectedItems[0]->Type == FConcertSessionTreeItem::EType::ActiveSession && GetController()->CanRenameActiveSession(SelectedItems[0]->ServerAdminEndpointId, SelectedItems[0]->SessionId)) ||
+	       (SelectedItems[0]->Type == FConcertSessionTreeItem::EType::ArchivedSession && GetController()->CanRenameArchivedSession(SelectedItems[0]->ServerAdminEndpointId, SelectedItems[0]->SessionId));
 }
 
 bool SConcertSessionBrowser::IsDeleteButtonEnabledInternal() const
 {
-	TArray<TSharedPtr<FConcertSessionItem>> SelectedItems = SessionsView->GetSelectedItems();
+	TArray<TSharedPtr<FConcertSessionTreeItem>> SelectedItems = GetSelectedItems();
 	if (SelectedItems.Num() == 0)
 	{
 		return false;
 	}
 
-	return (SelectedItems[0]->Type == FConcertSessionItem::EType::ActiveSession && GetController()->CanDeleteActiveSession(SelectedItems[0]->ServerAdminEndpointId, SelectedItems[0]->SessionId)) ||
-	       (SelectedItems[0]->Type == FConcertSessionItem::EType::ArchivedSession && GetController()->CanDeleteArchivedSession(SelectedItems[0]->ServerAdminEndpointId, SelectedItems[0]->SessionId));
+	return (SelectedItems[0]->Type == FConcertSessionTreeItem::EType::ActiveSession && GetController()->CanDeleteActiveSession(SelectedItems[0]->ServerAdminEndpointId, SelectedItems[0]->SessionId)) ||
+	       (SelectedItems[0]->Type == FConcertSessionTreeItem::EType::ArchivedSession && GetController()->CanDeleteArchivedSession(SelectedItems[0]->ServerAdminEndpointId, SelectedItems[0]->SessionId));
 }
 
 FReply SConcertSessionBrowser::OnNewButtonClicked()
@@ -1020,7 +1099,7 @@ FReply SConcertSessionBrowser::OnNewButtonClicked()
 
 FReply SConcertSessionBrowser::OnRestoreButtonClicked()
 {
-	TArray<TSharedPtr<FConcertSessionItem>> SelectedItems = SessionsView->GetSelectedItems();
+	const TArray<TSharedPtr<FConcertSessionTreeItem>> SelectedItems = GetSelectedItems();
 	if (SelectedItems.Num() == 1)
 	{
 		InsertRestoreSessionAsEditableRowInternal(SelectedItems[0]);
@@ -1031,7 +1110,7 @@ FReply SConcertSessionBrowser::OnRestoreButtonClicked()
 
 FReply SConcertSessionBrowser::OnArchiveButtonClicked()
 {
-	TArray<TSharedPtr<FConcertSessionItem>> SelectedItems = SessionsView->GetSelectedItems();
+	const TArray<TSharedPtr<FConcertSessionTreeItem>> SelectedItems = GetSelectedItems();
 	if (SelectedItems.Num() == 1)
 	{
 		InsertArchiveSessionAsEditableRow(SelectedItems[0]);
@@ -1042,11 +1121,11 @@ FReply SConcertSessionBrowser::OnArchiveButtonClicked()
 
 FReply SConcertSessionBrowser::OnDeleteButtonClicked()
 {
-	RequestDeleteSessions(SessionsView->GetSelectedItems());
+	RequestDeleteSessions(GetSelectedItems());
 	return FReply::Handled();
 }
 
-void SConcertSessionBrowser::OnBeginEditingSessionName(TSharedPtr<FConcertSessionItem> Item)
+void SConcertSessionBrowser::OnBeginEditingSessionName(TSharedPtr<FConcertSessionTreeItem> Item)
 {
 	Item->OnBeginEditSessionNameRequest.Broadcast(); // Signal the row widget to enter in edit mode.
 }
@@ -1076,18 +1155,20 @@ FReply SConcertSessionBrowser::OnKeyDown(const FGeometry& MyGeometry, const FKey
 
 	if (InKeyEvent.GetKey() == EKeys::Delete && !EditableSessionRow.IsValid()) // Delete selected row(s) unless the selected row is an 'editable' one.
 	{
-		RequestDeleteSessions(SessionsView->GetSelectedItems());
+		RequestDeleteSessions(GetSelectedItems());
 		return FReply::Handled();
 	}
-	else if (InKeyEvent.GetKey() == EKeys::Escape && EditableSessionRow.IsValid()) // Cancel 'new session', 'archive session' or 'restore session' action.
+	
+	if (InKeyEvent.GetKey() == EKeys::Escape && EditableSessionRow.IsValid()) // Cancel 'new session', 'archive session' or 'restore session' action.
 	{
 		RemoveSessionRow(EditableSessionRow);
 		check(!EditableSessionRow.IsValid() && !EditableSessionRowParent.IsValid()); // Expect to be cleared by RemoveSessionRow().
 		return FReply::Handled();
 	}
-	else if (InKeyEvent.GetKey() == EKeys::F2 && !EditableSessionRow.IsValid())
+	
+	if (InKeyEvent.GetKey() == EKeys::F2 && !EditableSessionRow.IsValid())
 	{
-		TArray<TSharedPtr<FConcertSessionItem>> SelectedItems = SessionsView->GetSelectedItems();
+		TArray<TSharedPtr<FConcertSessionTreeItem>> SelectedItems = GetSelectedItems();
 		if (SelectedItems.Num() == 1)
 		{
 			SelectedItems[0]->OnBeginEditSessionNameRequest.Broadcast(); // Broadcast the request.
@@ -1097,31 +1178,43 @@ FReply SConcertSessionBrowser::OnKeyDown(const FGeometry& MyGeometry, const FKey
 	return FReply::Unhandled();
 }
 
-void SConcertSessionBrowser::RequestCreateSession(const TSharedPtr<FConcertSessionItem>& NewItem)
+TArray<TSharedPtr<FConcertSessionTreeItem>> SConcertSessionBrowser::GetSelectedItems() const
+{
+	TArray<TSharedPtr<FConcertSessionTreeItem>> Result;
+	const TArray<TSharedPtr<FConcertTreeItem>> Selected = SessionsView->GetSelectedItems();
+	Algo::ForEachIf(
+		Selected,
+		[this](const TSharedPtr<FConcertTreeItem>& Item){ return !IsGroupItem(Item); },
+		[&Result](const TSharedPtr<FConcertTreeItem>& Item){ Result.Add(StaticCastSharedPtr<FConcertSessionTreeItem>(Item)); }
+		);
+	return Result;
+}
+
+void SConcertSessionBrowser::RequestCreateSession(const TSharedPtr<FConcertSessionTreeItem>& NewItem)
 {
 	GetController()->CreateSession(NewItem->ServerAdminEndpointId, NewItem->SessionName, NewItem->ProjectName);
 	RemoveSessionRow(NewItem); // The row used to edit the session name and pick the server.
 }
 
-void SConcertSessionBrowser::RequestArchiveSession(const TSharedPtr<FConcertSessionItem>& SaveItem, const FString& ArchiveName)
+void SConcertSessionBrowser::RequestArchiveSession(const TSharedPtr<FConcertSessionTreeItem>& SaveItem, const FString& ArchiveName)
 {
 	GetController()->ArchiveSession(SaveItem->ServerAdminEndpointId, SaveItem->SessionId, ArchiveName, FConcertSessionFilter());
 	RemoveSessionRow(SaveItem); // The row used to edit the archive name.
 }
 
-void SConcertSessionBrowser::RequestRestoreSession(const TSharedPtr<FConcertSessionItem>& RestoreItem, const FString& SessionName)
+void SConcertSessionBrowser::RequestRestoreSession(const TSharedPtr<FConcertSessionTreeItem>& RestoreItem, const FString& SessionName)
 {
 	GetController()->RestoreSession(RestoreItem->ServerAdminEndpointId, RestoreItem->SessionId, SessionName, FConcertSessionFilter());
 	RemoveSessionRow(RestoreItem); // The row used to edit the restore as name.
 }
 
-void SConcertSessionBrowser::RequestRenameSession(const TSharedPtr<FConcertSessionItem>& RenamedItem, const FString& NewName)
+void SConcertSessionBrowser::RequestRenameSession(const TSharedPtr<FConcertSessionTreeItem>& RenamedItem, const FString& NewName)
 {
-	if (RenamedItem->Type == FConcertSessionItem::EType::ActiveSession)
+	if (RenamedItem->Type == FConcertSessionTreeItem::EType::ActiveSession)
 	{
 		GetController()->RenameActiveSession(RenamedItem->ServerAdminEndpointId, RenamedItem->SessionId, NewName);
 	}
-	else if (RenamedItem->Type == FConcertSessionItem::EType::ArchivedSession)
+	else if (RenamedItem->Type == FConcertSessionTreeItem::EType::ArchivedSession)
 	{
 		GetController()->RenameArchivedSession(RenamedItem->ServerAdminEndpointId, RenamedItem->SessionId, NewName);
 	}
@@ -1131,7 +1224,7 @@ void SConcertSessionBrowser::RequestRenameSession(const TSharedPtr<FConcertSessi
 	RenamedItem->SessionName = NewName;
 }
 
-void SConcertSessionBrowser::RequestDeleteSessions(const TArray<TSharedPtr<FConcertSessionItem>>& ItemsToDelete)
+void SConcertSessionBrowser::RequestDeleteSessions(const TArray<TSharedPtr<FConcertSessionTreeItem>>& ItemsToDelete)
 {
 	if (!AskUserToDeleteSessions.IsBound() || AskUserToDeleteSessions.Execute(ItemsToDelete))
 	{
diff --git a/Engine/Plugins/Developer/Concert/ConcertUI/ConcertSharedSlate/Source/ConcertSharedSlate/Private/Session/Browser/SNewSessionRow.cpp b/Engine/Plugins/Developer/Concert/ConcertUI/ConcertSharedSlate/Source/ConcertSharedSlate/Private/Session/Browser/SNewSessionRow.cpp
--- a/Engine/Plugins/Developer/Concert/ConcertUI/ConcertSharedSlate/Source/ConcertSharedSlate/Private/Session/Browser/SNewSessionRow.cpp
+++ b/Engine/Plugins/Developer/Concert/ConcertUI/ConcertSharedSlate/Source/ConcertSharedSlate/Private/Session/Browser/SNewSessionRow.cpp
@@ -2,15 +2,15 @@
 
 #include "SNewSessionRow.h"
 
+#include "ConcertFrontendStyle.h"
 #include "Session/Browser/ConcertBrowserUtils.h"
-#include "Session/Browser/ConcertSessionItem.h"		
+#include "Session/Browser/Items/ConcertSessionTreeItem.h"		
 
 #include "Styling/AppStyle.h"
 
 #include "Algo/ForEach.h"
 
 #include "Framework/Application/SlateApplication.h"
-#include "Styling/AppStyle.h"
 #include "Widgets/Input/SEditableTextBox.h"
 #include "Widgets/Layout/SBox.h"
 #include "Widgets/SBoxPanel.h"
@@ -20,7 +20,7 @@
 
 #define LOCTEXT_NAMESPACE "SConcertBrowser"
 
-void SNewSessionRow::Construct(const FArguments& InArgs, TSharedPtr<FConcertSessionItem> InItem, const TSharedRef<STableViewBase>& InOwnerTableView)
+void SNewSessionRow::Construct(const FArguments& InArgs, TSharedPtr<FConcertSessionTreeItem> InItem, const TSharedRef<STableViewBase>& InOwnerTableView)
 {
 	Item = MoveTemp(InItem);
 	GetServersFunc = InArgs._GetServerFunc;
@@ -34,12 +34,11 @@ void SNewSessionRow::Construct(const FArguments& InArgs, TSharedPtr<FConcertSess
 	DeclineFunc.CheckCallable();
 
 	// Construct base class
-	SMultiColumnTableRow<TSharedPtr<FConcertSessionItem>>::Construct(FSuperRowType::FArguments(), InOwnerTableView);
+	SMultiColumnTableRow<TSharedPtr<FConcertSessionTreeItem>>::Construct(FSuperRowType::FArguments().Padding(InArgs._Padding), InOwnerTableView);
 	
 	// This rows are needed to correctly display the row
 	TemporaryColumnShower
 		.SetHeaderRow(InOwnerTableView->GetHeaderRow())
-		.SaveVisibilityAndSet(ConcertBrowserUtils::IconColName, true)
 		.SaveVisibilityAndSet(ConcertBrowserUtils::SessionColName, true)
 		.SaveVisibilityAndSet(ConcertBrowserUtils::ServerColName, true)
 		.SaveVisibilityAndSet(ConcertBrowserUtils::VersionColName, true)
@@ -68,33 +67,50 @@ void SNewSessionRow::Tick(const FGeometry& AllottedGeometry, const double InCurr
 
 TSharedRef<SWidget> SNewSessionRow::GenerateWidgetForColumn(const FName& ColumnName)
 {
-	TSharedPtr<FConcertSessionItem> ItemPin = Item.Pin();
+	TSharedPtr<FConcertSessionTreeItem> ItemPin = Item.Pin();
 
-	if (ColumnName == ConcertBrowserUtils::IconColName)
+	if (ColumnName == ConcertBrowserUtils::SessionColName)
 	{
-		// 'New' icon
-		return SNew(SBox)
-			.VAlign(VAlign_Center)
-			.HAlign(HAlign_Center)
+		return SNew( SHorizontalBox )
+
+			+SHorizontalBox::Slot()
+			.AutoWidth()
+			.Padding(6, 0, 0, 0)
 			[
-				SNew(SImage)
-				.Image(FAppStyle::GetBrush("Icons.PlusCircle"))
-			];
-	}
-	else if (ColumnName == ConcertBrowserUtils::SessionColName)
-	{
-		return SNew(SBox)
-			.VAlign(VAlign_Center)
-			.Padding(FMargin(0, 0, 2, 0))
+				SNew( SExpanderArrow, SharedThis(this) ).IndentAmount(12)
+			]
+
+			+SHorizontalBox::Slot()
+			.FillWidth(1.0f)
 			[
-				SAssignNew(EditableSessionName, SEditableTextBox)
-				.HintText(LOCTEXT("EnterSessionNameHint", "Enter a session name"))
-				.OnTextCommitted(this, &SNewSessionRow::OnSessionNameCommitted)
-				.OnKeyDownHandler(this, &SNewSessionRow::OnKeyDownHandler)
-				.OnTextChanged(this, &SNewSessionRow::OnSessionNameChanged)
+				SNew(SBox)
+				.VAlign(VAlign_Center)
+				.Padding(FMargin(0, 0, 2, 0))
+				[
+					SNew(SHorizontalBox)
+
+					+SHorizontalBox::Slot()
+					.AutoWidth()
+					.HAlign(HAlign_Center)
+					.VAlign(VAlign_Center)
+					.Padding(1.f, 1.f, 4.f, 1.f)
+					[
+						SNew(SImage)
+						.Image(FAppStyle::GetBrush("Icons.PlusCircle"))
+					]
+					
+					+SHorizontalBox::Slot()
+					[
+						SAssignNew(EditableSessionName, SEditableTextBox)
+						.HintText(LOCTEXT("EnterSessionNameHint", "Enter a session name"))
+						.OnTextCommitted(this, &SNewSessionRow::OnSessionNameCommitted)
+						.OnKeyDownHandler(this, &SNewSessionRow::OnKeyDownHandler)
+						.OnTextChanged(this, &SNewSessionRow::OnSessionNameChanged)
+					]
+				]
 			];
 	}
-	else if (ColumnName == ConcertBrowserUtils::ServerColName)
+	if (ColumnName == ConcertBrowserUtils::ServerColName)
 	{
 		return SNew(SHorizontalBox)
 
@@ -112,7 +128,7 @@ TSharedRef<SWidget> SNewSessionRow::GenerateWidgetForColumn(const FName& ColumnN
 				]
 			];
 	}
-	else if (ColumnName == ConcertBrowserUtils::VersionColName)
+	if (ColumnName == ConcertBrowserUtils::VersionColName)
 	{
 		return SNew(SHorizontalBox)
 			+SHorizontalBox::Slot()
@@ -144,7 +160,7 @@ TSharedRef<SWidget> SNewSessionRow::GenerateWidgetForColumn(const FName& ColumnN
 				]
 			];
 	}
-	else if (ColumnName == ConcertBrowserUtils::ProjectColName)
+	if (ColumnName == ConcertBrowserUtils::ProjectColName)
 	{
 		FText ProjectName = FText::AsCultureInvariant(FApp::GetProjectName());
 		return SNew(SBox)
@@ -281,6 +297,7 @@ TSharedRef<SWidget> SNewSessionRow::MakeSelectedServerWidget()
 			SNew(STextBlock)
 			.Text(this, &SNewSessionRow::GetSelectedServerText)
 			.HighlightText(HighlightText)
+			.Font(FCoreStyle::GetDefaultFontStyle("Regular", FConcertFrontendStyle::Get()->GetFloat("Concert.SessionBrowser.FontSize")))
 		]
 		+SHorizontalBox::Slot()
 		.AutoWidth()
@@ -336,7 +353,7 @@ FText SNewSessionRow::GetSelectedServerIgnoreVersionTooltip() const
 
 FReply SNewSessionRow::OnAccept()
 {
-	TSharedPtr<FConcertSessionItem> ItemPin = Item.Pin();
+	TSharedPtr<FConcertSessionTreeItem> ItemPin = Item.Pin();
 	FString NewSessionName = EditableSessionName->GetText().ToString();
 	FText InvalidNameErrorMsg = ConcertSettingsUtils::ValidateSessionName(NewSessionName);
 
diff --git a/Engine/Plugins/Developer/Concert/ConcertUI/ConcertSharedSlate/Source/ConcertSharedSlate/Private/Session/Browser/SNewSessionRow.h b/Engine/Plugins/Developer/Concert/ConcertUI/ConcertSharedSlate/Source/ConcertSharedSlate/Private/Session/Browser/SNewSessionRow.h
--- a/Engine/Plugins/Developer/Concert/ConcertUI/ConcertSharedSlate/Source/ConcertSharedSlate/Private/Session/Browser/SNewSessionRow.h
+++ b/Engine/Plugins/Developer/Concert/ConcertUI/ConcertSharedSlate/Source/ConcertSharedSlate/Private/Session/Browser/SNewSessionRow.h
@@ -6,11 +6,12 @@
 #include "ConcertHeaderRowUtils.h"
 #include "ConcertMessageData.h"
 #include "Widgets/DeclarativeSyntaxSupport.h"
-#include "Widgets/Views/STableRow.h"
 #include "Widgets/Input/SComboBox.h"
+#include "Widgets/Views/STableRow.h"
+#include "Widgets/Views/STreeView.h"
 #include "Styling/SlateBrush.h"
 
-class FConcertSessionItem;
+class FConcertSessionTreeItem;
 class SEditableTextBox;
 class SHeaderRow;
 class SWidget;
@@ -18,14 +19,14 @@ class SWidget;
 /**
  * The type of row used in the session list view to edit a new session (the session name + server).
  */
-class SNewSessionRow : public SMultiColumnTableRow<TSharedPtr<FConcertSessionItem>>
+class SNewSessionRow : public SMultiColumnTableRow<TSharedPtr<FConcertSessionTreeItem>>
 {
 public:
 	using FGetServersFunc = TFunction<TArray<FConcertServerInfo>()>;
 	// Should remove the editable 'new' row and creates the sessions.
-	using FAcceptFunc = TFunction<void(const TSharedPtr<FConcertSessionItem>&)>;
+	using FAcceptFunc = TFunction<void(const TSharedPtr<FConcertSessionTreeItem>&)>;
 	// Should just remove the editable 'new' row.
-	using FDeclineFunc = TFunction<void(const TSharedPtr<FConcertSessionItem>&)>; 
+	using FDeclineFunc = TFunction<void(const TSharedPtr<FConcertSessionTreeItem>&)>; 
 
 	SLATE_BEGIN_ARGS(SNewSessionRow) // Needed to use Args because Construct() is limited to 5 arguments and 6 were required.
 		: _GetServerFunc()
@@ -39,9 +40,11 @@ public:
 
 		SLATE_ATTRIBUTE(FText, HighlightText)
 		SLATE_ATTRIBUTE(FString, DefaultServerURL)
+	
+		SLATE_ATTRIBUTE(FMargin, Padding)
 	SLATE_END_ARGS()
 	
-	void Construct(const FArguments& InArgs, TSharedPtr<FConcertSessionItem> InItem, const TSharedRef<STableViewBase>& InOwnerTableView);
+	void Construct(const FArguments& InArgs, TSharedPtr<FConcertSessionTreeItem> InItem, const TSharedRef<STableViewBase>& InOwnerTableView);
 	virtual ~SNewSessionRow() override;
 	
 	virtual TSharedRef<SWidget> GenerateWidgetForColumn(const FName& ColumnName) override;
@@ -70,7 +73,7 @@ private:
 	UE::ConcertSharedSlate::FColumnVisibilityTransaction TemporaryColumnShower;
 	
 	/** Holds the new item to fill with session name and server. */
-	TWeakPtr<FConcertSessionItem> Item;
+	TWeakPtr<FConcertSessionTreeItem> Item;
 	/** Servers displayed in the server combo box. */
 	TArray<TSharedPtr<FConcertServerInfo>> Servers;
 
diff --git a/Engine/Plugins/Developer/Concert/ConcertUI/ConcertSharedSlate/Source/ConcertSharedSlate/Private/Session/Browser/SSaveRestoreSessionRow.cpp b/Engine/Plugins/Developer/Concert/ConcertUI/ConcertSharedSlate/Source/ConcertSharedSlate/Private/Session/Browser/SSaveRestoreSessionRow.cpp
--- a/Engine/Plugins/Developer/Concert/ConcertUI/ConcertSharedSlate/Source/ConcertSharedSlate/Private/Session/Browser/SSaveRestoreSessionRow.cpp
+++ b/Engine/Plugins/Developer/Concert/ConcertUI/ConcertSharedSlate/Source/ConcertSharedSlate/Private/Session/Browser/SSaveRestoreSessionRow.cpp
@@ -2,8 +2,9 @@
 
 #include "SSaveRestoreSessionRow.h"
 
+#include "ConcertFrontendStyle.h"
 #include "Session/Browser/ConcertBrowserUtils.h"
-#include "Session/Browser/ConcertSessionItem.h"
+#include "Session/Browser/Items/ConcertSessionTreeItem.h"
 
 #include "Styling/AppStyle.h"
 
@@ -11,7 +12,6 @@
 
 #include "Framework/Application/SlateApplication.h"
 #include "Internationalization/Regex.h"
-#include "Styling/AppStyle.h"
 #include "Widgets/Input/SEditableTextBox.h"
 #include "Widgets/Layout/SBox.h"
 #include "Widgets/Layout/SUniformGridPanel.h"
@@ -21,7 +21,7 @@
 
 #define LOCTEXT_NAMESPACE "SConcertBrowser"
 
-void SSaveRestoreSessionRow::Construct(const FArguments& InArgs, TSharedPtr<FConcertSessionItem> InNode, const TSharedRef<STableViewBase>& InOwnerTableView)
+void SSaveRestoreSessionRow::Construct(const FArguments& InArgs, TSharedPtr<FConcertSessionTreeItem> InNode, const TSharedRef<STableViewBase>& InOwnerTableView)
 {
 	Item = MoveTemp(InNode);
 	AcceptFunc = InArgs._OnAcceptFunc;
@@ -32,12 +32,11 @@ void SSaveRestoreSessionRow::Construct(const FArguments& InArgs, TSharedPtr<FCon
 	DeclineFunc.CheckCallable();
 
 	// Construct base class
-	SMultiColumnTableRow<TSharedPtr<FConcertSessionItem>>::Construct(FSuperRowType::FArguments(), InOwnerTableView);
+	SMultiColumnTableRow<TSharedPtr<FConcertSessionTreeItem>>::Construct(FSuperRowType::FArguments().Padding(InArgs._Padding), InOwnerTableView);
 
 	// This rows are needed to correctly display the row
 	TemporaryColumnShower
 		.SetHeaderRow(InOwnerTableView->GetHeaderRow())
-		.SaveVisibilityAndSet(ConcertBrowserUtils::IconColName, true)
 		.SaveVisibilityAndSet(ConcertBrowserUtils::SessionColName, true)
 		.SaveVisibilityAndSet(ConcertBrowserUtils::ServerColName, true)
 		.SaveVisibilityAndSet(ConcertBrowserUtils::ProjectColName, true);
@@ -55,9 +54,9 @@ TBitArray<> SSaveRestoreSessionRow::GetWiresNeededByDepth() const
 	return Bits;
 }
 
-FText SSaveRestoreSessionRow::GetDefaultName(const FConcertSessionItem& InItem) const
+FText SSaveRestoreSessionRow::GetDefaultName(const FConcertSessionTreeItem& InItem) const
 {
-	if (InItem.Type == FConcertSessionItem::EType::SaveSession)
+	if (InItem.Type == FConcertSessionTreeItem::EType::SaveSession)
 	{
 		return FText::Format(LOCTEXT("DefaultName", "{0}.{1}"), FText::FromString(InItem.SessionName), FText::FromString(FDateTime::UtcNow().ToString()));
 	}
@@ -75,48 +74,62 @@ FText SSaveRestoreSessionRow::GetDefaultName(const FConcertSessionItem& InItem)
 
 TSharedRef<SWidget> SSaveRestoreSessionRow::GenerateWidgetForColumn(const FName& ColumnName)
 {
-	TSharedPtr<FConcertSessionItem> ItemPin = Item.Pin();
+	TSharedPtr<FConcertSessionTreeItem> ItemPin = Item.Pin();
 
-	if (ColumnName == ConcertBrowserUtils::IconColName)
-	{
-		return SNew(SHorizontalBox)
-			+SHorizontalBox::Slot()
-			.AutoWidth()
-			.Padding(8, 0)
-			[
-				SNew(SExpanderArrow, SharedThis(this))
-				.StyleSet(&FAppStyle::Get())
-				.ShouldDrawWires(true)
-			];
-	}
-	else if (ColumnName == ConcertBrowserUtils::SessionColName)
+	if (ColumnName == ConcertBrowserUtils::SessionColName)
 	{
-		return SNew(SHorizontalBox)
+		return SNew( SHorizontalBox )
 
-			// 'Restore as/Save as' text
 			+SHorizontalBox::Slot()
-			.VAlign(VAlign_Center)
 			.AutoWidth()
-			.Padding(2.0, 0.0)
+			.Padding(6, 0, 0, 0)
 			[
-				SNew(STextBlock).Text(ItemPin->Type == FConcertSessionItem::EType::RestoreSession ? LOCTEXT("RestoreAs", "Restore as:") : LOCTEXT("ArchiveAs", "Archive as:"))
+				SNew( SExpanderArrow, SharedThis(this) ).IndentAmount(12)
 			]
 
-			// Editable text.
 			+SHorizontalBox::Slot()
-			.VAlign(VAlign_Center)
-			.Padding(FMargin(0, 0, 2, 0))
+			.FillWidth(1.0f)
 			[
-				SAssignNew(EditableSessionName, SEditableTextBox)
-				.HintText(ItemPin->Type == FConcertSessionItem::EType::RestoreSession ? LOCTEXT("RestoreSessionHint", "Enter a session name") : LOCTEXT("ArchivSessionHint", "Enter an archive name"))
-				.OnTextCommitted(this, &SSaveRestoreSessionRow::OnSessionNameCommitted)
-				.OnKeyDownHandler(this, &SSaveRestoreSessionRow::OnKeyDownHandler)
-				.OnTextChanged(this, &SSaveRestoreSessionRow::OnSessionNameChanged)
-				.Text(GetDefaultName(*ItemPin))
-				.SelectAllTextWhenFocused(true)
+				SNew(SHorizontalBox)
+
+				+SHorizontalBox::Slot()
+				.AutoWidth()
+				.Padding(8, 0)
+				[
+					SNew(SExpanderArrow, SharedThis(this))
+					.IndentAmount(4)
+					.StyleSet(&FAppStyle::Get())
+					.ShouldDrawWires(true)
+				]
+				
+				// 'Restore as/Save as' text
+				+SHorizontalBox::Slot()
+				.VAlign(VAlign_Center)
+				.AutoWidth()
+				.Padding(2.0, 0.0)
+				[
+					SNew(STextBlock)
+					.Text(ItemPin->Type == FConcertSessionTreeItem::EType::RestoreSession ? LOCTEXT("RestoreAs", "Restore as:") : LOCTEXT("ArchiveAs", "Archive as:"))
+					.Font(FCoreStyle::GetDefaultFontStyle("Regular", FConcertFrontendStyle::Get()->GetFloat("Concert.SessionBrowser.FontSize")))
+				]
+
+				// Editable text.
+				+SHorizontalBox::Slot()
+				.VAlign(VAlign_Center)
+				.Padding(FMargin(0, 0, 2, 0))
+				[
+					SAssignNew(EditableSessionName, SEditableTextBox)
+					.HintText(ItemPin->Type == FConcertSessionTreeItem::EType::RestoreSession ? LOCTEXT("RestoreSessionHint", "Enter a session name") : LOCTEXT("ArchivSessionHint", "Enter an archive name"))
+					.OnTextCommitted(this, &SSaveRestoreSessionRow::OnSessionNameCommitted)
+					.OnKeyDownHandler(this, &SSaveRestoreSessionRow::OnKeyDownHandler)
+					.OnTextChanged(this, &SSaveRestoreSessionRow::OnSessionNameChanged)
+					.Text(GetDefaultName(*ItemPin))
+					.SelectAllTextWhenFocused(true)
+					.Font(FCoreStyle::GetDefaultFontStyle("Regular", FConcertFrontendStyle::Get()->GetFloat("Concert.SessionBrowser.FontSize")))
+				]
 			];
 	}
-	else if (ColumnName == ConcertBrowserUtils::ServerColName)
+	if (ColumnName == ConcertBrowserUtils::ServerColName)
 	{
 		// 'Server' text block.
 		return SNew(SHorizontalBox)
@@ -129,11 +142,12 @@ TSharedRef<SWidget> SSaveRestoreSessionRow::GenerateWidgetForColumn(const FName&
 				// Server.
 				SNew(STextBlock)
 				.Text(FText::FromString(ItemPin->ServerName))
+				.Font(FCoreStyle::GetDefaultFontStyle("Regular", FConcertFrontendStyle::Get()->GetFloat("Concert.SessionBrowser.FontSize")))
 				.HighlightText(HighlightText)
 			]
 		];
 	}
-	else if (ColumnName == ConcertBrowserUtils::ProjectColName)
+	if (ColumnName == ConcertBrowserUtils::ProjectColName)
 	{
 		return SNew(SHorizontalBox)
 			+SHorizontalBox::Slot()
@@ -149,7 +163,7 @@ TSharedRef<SWidget> SSaveRestoreSessionRow::GenerateWidgetForColumn(const FName&
 				[
 					ConcertBrowserUtils::MakePositiveActionButton(
 						FAppStyle::GetBrush("Icons.Check"),
-						ItemPin->Type == FConcertSessionItem::EType::RestoreSession ? LOCTEXT("RestoreCheckIconTooltip", "Restore the session") : LOCTEXT("ArchiveCheckIconTooltip", "Archive the session"),
+						ItemPin->Type == FConcertSessionTreeItem::EType::RestoreSession ? LOCTEXT("RestoreCheckIconTooltip", "Restore the session") : LOCTEXT("ArchiveCheckIconTooltip", "Archive the session"),
 						TAttribute<bool>::Create([this]() { return !EditableSessionName->GetText().IsEmpty(); }), // Enabled?
 						FOnClicked::CreateRaw(this, &SSaveRestoreSessionRow::OnAccept))
 				]
@@ -188,7 +202,7 @@ void SSaveRestoreSessionRow::OnSessionNameChanged(const FText& NewName)
 
 void SSaveRestoreSessionRow::OnSessionNameCommitted(const FText& NewText, ETextCommit::Type CommitType)
 {
-	TSharedPtr<FConcertSessionItem> ItemPin = Item.Pin();
+	TSharedPtr<FConcertSessionTreeItem> ItemPin = Item.Pin();
 	if (CommitType == ETextCommit::Type::OnEnter)
 	{
 		OnAccept();
@@ -198,7 +212,7 @@ void SSaveRestoreSessionRow::OnSessionNameCommitted(const FText& NewText, ETextC
 FReply SSaveRestoreSessionRow::OnAccept()
 {
 	// Read the session name given by the user.
-	TSharedPtr<FConcertSessionItem> ItemPin = Item.Pin();
+	TSharedPtr<FConcertSessionTreeItem> ItemPin = Item.Pin();
 	FString Name = EditableSessionName->GetText().ToString(); // Archive name or restored session name.
 
 	// Ensure the user provided a name.
diff --git a/Engine/Plugins/Developer/Concert/ConcertUI/ConcertSharedSlate/Source/ConcertSharedSlate/Private/Session/Browser/SSaveRestoreSessionRow.h b/Engine/Plugins/Developer/Concert/ConcertUI/ConcertSharedSlate/Source/ConcertSharedSlate/Private/Session/Browser/SSaveRestoreSessionRow.h
--- a/Engine/Plugins/Developer/Concert/ConcertUI/ConcertSharedSlate/Source/ConcertSharedSlate/Private/Session/Browser/SSaveRestoreSessionRow.h
+++ b/Engine/Plugins/Developer/Concert/ConcertUI/ConcertSharedSlate/Source/ConcertSharedSlate/Private/Session/Browser/SSaveRestoreSessionRow.h
@@ -4,12 +4,12 @@
 
 #include "CoreMinimal.h"
 #include "ConcertHeaderRowUtils.h"
+#include "Session/Browser/Items/ConcertSessionTreeItem.h"
 
 #include "Misc/Attribute.h"
-#include "Session/Browser/ConcertSessionItem.h"
-#include "Widgets/Views/SListView.h"
 #include "Widgets/DeclarativeSyntaxSupport.h"
 #include "Widgets/Views/STableRow.h"
+#include "Widgets/Views/STreeView.h"
 
 class SWidget;
 class SEditableTextBox;
@@ -17,13 +17,13 @@ class SEditableTextBox;
 /**
  * The type of row used in the session list view to create and archive or restore a session (edit the information required for the operation).
  */
-class SSaveRestoreSessionRow : public SMultiColumnTableRow<TSharedPtr<FConcertSessionItem>>
+class SSaveRestoreSessionRow : public SMultiColumnTableRow<TSharedPtr<FConcertSessionTreeItem>>
 {
 public:
 	/** Should remove the editable row and save or restore the session. */
-	using FAcceptFunc = TFunction<void(TSharedPtr<FConcertSessionItem>, const FString&)>;
+	using FAcceptFunc = TFunction<void(TSharedPtr<FConcertSessionTreeItem>, const FString&)>;
 	/** Should only remove the editable row from the table. */
-	using FDeclineFunc = TFunction<void(TSharedPtr<FConcertSessionItem>)>;
+	using FDeclineFunc = TFunction<void(TSharedPtr<FConcertSessionTreeItem>)>;
 
 	SLATE_BEGIN_ARGS(SSaveRestoreSessionRow)
 		: _OnAcceptFunc()
@@ -33,9 +33,10 @@ public:
 		SLATE_ARGUMENT(FAcceptFunc, OnAcceptFunc)
 		SLATE_ARGUMENT(FDeclineFunc, OnDeclineFunc)
 		SLATE_ATTRIBUTE(FText, HighlightText)
+		SLATE_ATTRIBUTE(FMargin, Padding)
 	SLATE_END_ARGS()
 	
-	void Construct(const FArguments& InArgs, TSharedPtr<FConcertSessionItem> Node, const TSharedRef<STableViewBase>& InOwnerTableView);
+	void Construct(const FArguments& InArgs, TSharedPtr<FConcertSessionTreeItem> Node, const TSharedRef<STableViewBase>& InOwnerTableView);
 	virtual ~SSaveRestoreSessionRow() override;
 	
 	virtual TSharedRef<SWidget> GenerateWidgetForColumn(const FName& ColumnName) override;
@@ -54,14 +55,14 @@ public:
 	virtual bool IsItemExpanded() const override { return false; }
 
 	/** Generates a default name for an archive or a restored session. */
-	FText GetDefaultName(const FConcertSessionItem& Item) const; 
+	FText GetDefaultName(const FConcertSessionTreeItem& Item) const; 
 
 private:
 	
 	/** Helper for temporarily showing all columns required for this column */
 	UE::ConcertSharedSlate::FColumnVisibilityTransaction TemporaryColumnShower;
 	
-	TWeakPtr<FConcertSessionItem> Item;
+	TWeakPtr<FConcertSessionTreeItem> Item;
 	TSharedPtr<SEditableTextBox> EditableSessionName;
 	
 	FAcceptFunc AcceptFunc;
diff --git a/Engine/Plugins/Developer/Concert/ConcertUI/ConcertSharedSlate/Source/ConcertSharedSlate/Private/Session/Browser/SSessionGroupRow.cpp b/Engine/Plugins/Developer/Concert/ConcertUI/ConcertSharedSlate/Source/ConcertSharedSlate/Private/Session/Browser/SSessionGroupRow.cpp
new file mode 100644
--- /dev/null
+++ b/Engine/Plugins/Developer/Concert/ConcertUI/ConcertSharedSlate/Source/ConcertSharedSlate/Private/Session/Browser/SSessionGroupRow.cpp
@@ -0,0 +1,62 @@
+// Copyright Epic Games, Inc. All Rights Reserved.
+
+#include "SSessionGroupRow.h"
+
+#include "ConcertFrontendStyle.h"
+#include "Session/Browser/ConcertBrowserUtils.h"
+#include "Session/Browser/Items/ConcertArchivedGroupTreeItem.h"
+
+#include "Widgets/SNullWidget.h"
+#include "Widgets/Text/STextBlock.h"
+
+#define LOCTEXT_NAMESPACE "SConcertBrowser.SSessionGroupRow"
+
+void SSessionGroupRow::Construct(const FArguments& InArgs, const TSharedRef<STableViewBase>& InOwnerTableView)
+{
+	GroupType = InArgs._GroupType;
+	SMultiColumnTableRow<TSharedPtr<FConcertSessionTreeItem>>::Construct(FSuperRowType::FArguments().Padding(InArgs._Padding), InOwnerTableView);
+}
+
+TSharedRef<SWidget> SSessionGroupRow::GenerateWidgetForColumn(const FName& ColumnName)
+{
+	if (ColumnName == ConcertBrowserUtils::SessionColName)
+	{
+		const bool bIsActive = GroupType == EGroupType::Active;
+		const FText Label = bIsActive
+			? LOCTEXT("ActiveSessionsLabel", "Active Sessions")
+			: LOCTEXT("ArchivedSessionsLabel", "Archived Sessions");
+		const FSlateBrush* Glyph = bIsActive
+			? FConcertFrontendStyle::Get()->GetBrush("Concert.ActiveSession.Icon")
+			: FConcertFrontendStyle::Get()->GetBrush("Concert.ArchivedSession.Icon");
+		return SNew( SHorizontalBox )
+
+			+SHorizontalBox::Slot()
+			.AutoWidth()
+			.Padding(6, 0, 0, 0)
+			[
+				SNew( SExpanderArrow, SharedThis(this) ).IndentAmount(12)
+			]
+
+			+SHorizontalBox::Slot()
+			.AutoWidth()
+			.Padding(0.f, 0.f, 4.f, 0.f)
+			[
+				SNew(SImage)
+				.Image(Glyph)
+			]
+		
+			+SHorizontalBox::Slot()
+			.FillWidth(1.0f)
+			.HAlign(HAlign_Left)
+			.VAlign(VAlign_Center)
+			[
+				SNew(STextBlock)
+				.Text(Label)
+				.Font(FCoreStyle::GetDefaultFontStyle("Regular", FConcertFrontendStyle::Get()->GetFloat("Concert.SessionBrowser.FontSize")))
+			];
+	}
+
+	return SNullWidget::NullWidget;
+}
+
+#undef LOCTEXT_NAMESPACE
diff --git a/Engine/Plugins/Developer/Concert/ConcertUI/ConcertSharedSlate/Source/ConcertSharedSlate/Private/Session/Browser/SSessionGroupRow.h b/Engine/Plugins/Developer/Concert/ConcertUI/ConcertSharedSlate/Source/ConcertSharedSlate/Private/Session/Browser/SSessionGroupRow.h
new file mode 100644
--- /dev/null
+++ b/Engine/Plugins/Developer/Concert/ConcertUI/ConcertSharedSlate/Source/ConcertSharedSlate/Private/Session/Browser/SSessionGroupRow.h
@@ -0,0 +1,43 @@
+// Copyright Epic Games, Inc. All Rights Reserved.
+
+#pragma once
+
+#include "CoreMinimal.h"
+#include "Widgets/DeclarativeSyntaxSupport.h"
+#include "Widgets/Views/STableRow.h"
+#include "Styling/SlateBrush.h"
+#include "Widgets/Views/STreeView.h"
+
+class FConcertArchivedGroupTreeItem;
+class FConcertSessionTreeItem;
+class SEditableTextBox;
+class SHeaderRow;
+class SWidget;
+
+/**
+ * The type of row used in the session list view to edit a new session (the session name + server).
+ */
+class SSessionGroupRow : public SMultiColumnTableRow<TSharedPtr<FConcertSessionTreeItem>>
+{
+public:
+
+	enum class EGroupType
+	{
+		Active,
+		Archived
+	};
+
+	SLATE_BEGIN_ARGS(SSessionGroupRow) 
+	{}
+		SLATE_ARGUMENT(EGroupType, GroupType)
+		SLATE_ATTRIBUTE(FMargin, Padding)
+	SLATE_END_ARGS()
+	
+	void Construct(const FArguments& InArgs, const TSharedRef<STableViewBase>& InOwnerTableView);
+	
+	virtual TSharedRef<SWidget> GenerateWidgetForColumn(const FName& ColumnName) override;
+
+private:
+	
+	EGroupType GroupType{};
+};
diff --git a/Engine/Plugins/Developer/Concert/ConcertUI/ConcertSharedSlate/Source/ConcertSharedSlate/Private/Session/Browser/SSessionRow.cpp b/Engine/Plugins/Developer/Concert/ConcertUI/ConcertSharedSlate/Source/ConcertSharedSlate/Private/Session/Browser/SSessionRow.cpp
--- a/Engine/Plugins/Developer/Concert/ConcertUI/ConcertSharedSlate/Source/ConcertSharedSlate/Private/Session/Browser/SSessionRow.cpp
+++ b/Engine/Plugins/Developer/Concert/ConcertUI/ConcertSharedSlate/Source/ConcertSharedSlate/Private/Session/Browser/SSessionRow.cpp
@@ -3,14 +3,11 @@
 #include "SSessionRow.h"
 
 #include "ConcertFrontendUtils.h"
-
 #include "Session/Browser/ConcertBrowserUtils.h"
-#include "Session/Browser/ConcertSessionItem.h"
-
-#include "Framework/Application/SlateApplication.h"
-
 #include "Session/Browser/ConcertSessionBrowserSettings.h"
+#include "Session/Browser/Items/ConcertSessionTreeItem.h"
 
+#include "Framework/Application/SlateApplication.h"
 #include "Styling/AppStyle.h"
 #include "Styling/CoreStyle.h"
 #include "Widgets/Layout/SBox.h"
@@ -21,7 +18,7 @@
 
 #define LOCTEXT_NAMESPACE "SConcertBrowser"
 
-void SSessionRow::Construct(const FArguments& InArgs, TSharedPtr<FConcertSessionItem> InItem, const TSharedRef<STableViewBase>& InOwnerTableView)
+void SSessionRow::Construct(const FArguments& InArgs, TSharedPtr<FConcertSessionTreeItem> InItem, const TSharedRef<STableViewBase>& InOwnerTableView)
 {
 	Item = MoveTemp(InItem);
 	DoubleClickFunc = InArgs._OnDoubleClickFunc; // This function should join a session or add a row to restore an archive.
@@ -35,7 +32,7 @@ void SSessionRow::Construct(const FArguments& InArgs, TSharedPtr<FConcertSession
 	IsDefaultSession.CheckCallable();
 
 	// Construct base class
-	SMultiColumnTableRow<TSharedPtr<FConcertSessionItem>>::Construct(FSuperRowType::FArguments(), InOwnerTableView);
+	SMultiColumnTableRow<TSharedPtr<FConcertSessionTreeItem>>::Construct(FSuperRowType::FArguments().Padding(InArgs._Padding), InOwnerTableView);
 
 	// Listen and handle rename request.
 	InItem->OnBeginEditSessionNameRequest.AddSP(this, &SSessionRow::OnBeginEditingSessionName);
@@ -57,9 +54,9 @@ FSlateFontInfo SSessionRow::GetFontInfo(bool bIsActiveSession, bool bIsDefault)
 {
 	if (bIsActiveSession)
 	{
-		return FAppStyle::Get().GetFontStyle("NormalFont");
+		return FCoreStyle::GetDefaultFontStyle("Regular", FConcertFrontendStyle::Get()->GetFloat("Concert.SessionBrowser.FontSize"));
 	}
-	return  FCoreStyle::GetDefaultFontStyle("Italic", 9);
+	return FCoreStyle::GetDefaultFontStyle("Italic", FConcertFrontendStyle::Get()->GetFloat("Concert.SessionBrowser.FontSize"));
 }
 
 TSharedRef<SWidget> SSessionRow::GenerateSessionColumn(const FSlateFontInfo& FontInfo, const FSlateColor& FontColor)
@@ -81,7 +78,7 @@ TSharedRef<SWidget> SSessionRow::GenerateSessionColumn(const FSlateFontInfo& Fon
 
 TSharedRef<SWidget> SSessionRow::GenerateServerColumn(const FSlateFontInfo& FontInfo, const FSlateColor& FontColor)
 {
-	TSharedPtr<FConcertSessionItem> ItemPin = Item.Pin();
+	TSharedPtr<FConcertSessionTreeItem> ItemPin = Item.Pin();
 	return SNew(SHorizontalBox)
 		+SHorizontalBox::Slot()
 		.AutoWidth()
@@ -108,7 +105,7 @@ TSharedRef<SWidget> SSessionRow::GenerateServerColumn(const FSlateFontInfo& Font
 
 TSharedRef<SWidget> SSessionRow::GenerateServerDefaultColumn(const FSlateFontInfo& FontInfo, const FSlateColor& FontColor)
 {
-	TSharedPtr<FConcertSessionItem> ItemPin = Item.Pin();
+	TSharedPtr<FConcertSessionTreeItem> ItemPin = Item.Pin();
 	return SNew(SHorizontalBox)
 		+SHorizontalBox::Slot()
 		.AutoWidth()
@@ -127,7 +124,7 @@ TSharedRef<SWidget> SSessionRow::GenerateServerDefaultColumn(const FSlateFontInf
 			SNew(STextBlock)
 			.Text(LOCTEXT("DefaultServerSession", "(Default Session/Server)"))
 			.HighlightText(HighlightText)
-			.Font(FCoreStyle::GetDefaultFontStyle("Regular", 9))
+			.Font(FCoreStyle::GetDefaultFontStyle("Regular", FConcertFrontendStyle::Get()->GetFloat("Concert.SessionBrowser.FontSize")))
 			.ColorAndOpacity(FontColor)
 		]
 		+SHorizontalBox::Slot()
@@ -148,7 +145,7 @@ TSharedRef<SWidget> SSessionRow::GenerateProjectColumn(const FSlateFontInfo &Fon
 	return SNew(SBox)
 		.VAlign(VAlign_Center)
 		[
-			SAssignNew(SessionNameText, SInlineEditableTextBlock)
+			SNew(SInlineEditableTextBlock)
 			.Text_Lambda([this]() { return FText::AsCultureInvariant(Item.Pin()->ProjectName); })
 			.HighlightText(HighlightText)
 			.IsReadOnly(true)
@@ -162,7 +159,7 @@ TSharedRef<SWidget> SSessionRow::GenerateVersionColumn(const FSlateFontInfo& Fon
 	return SNew(SBox)
 		.VAlign(VAlign_Center)
 		[
-			SAssignNew(SessionNameText, SInlineEditableTextBlock)
+			SNew(SInlineEditableTextBlock)
 			.Text_Lambda([this]() { return FText::AsCultureInvariant(Item.Pin()->ProjectVersion); })
 			.HighlightText(HighlightText)
 			.IsReadOnly(true)
@@ -183,44 +180,41 @@ TSharedRef<SWidget> SSessionRow::GenerateLastModifiedColumn(const FSlateFontInfo
 		];
 }
 
-
-TSharedRef<SWidget> SSessionRow::GenerateIconColumn()
-{
-	TSharedPtr<FConcertSessionItem> ItemPin = Item.Pin();
-	const FText Tooltip = ItemPin->Type == FConcertSessionItem::EType::ActiveSession ?
-		LOCTEXT("ActiveIconTooltip", "Active session") : LOCTEXT("ArchivedIconTooltip", "Archived Session");
-	const FSlateColor InStyle = ItemPin->Type == FConcertSessionItem::EType::ActiveSession ?
-		FAppStyle::Get().GetWidgetStyle<FButtonStyle>("FlatButton.Success").Normal.TintColor : FSlateColor::UseSubduedForeground();
-	// @todo LB: change FolderClosed to archive
-	const FSlateBrush* InGlyph = ItemPin->Type == FConcertSessionItem::EType::ActiveSession ? FAppStyle::GetBrush("Icons.FilledCircle") : FAppStyle::GetBrush("Icons.FolderClosed");
-	return SNew(SBox)
-		.VAlign(VAlign_Center)
-		.HAlign(HAlign_Center)
-		.Padding(2)
-		.ToolTipText(Tooltip)
-		[
-			SNew(SImage)
-			.Image(InGlyph)
-			.ColorAndOpacity(InStyle)
-		];
-}
-
 TSharedRef<SWidget> SSessionRow::GenerateWidgetForColumn(const FName& ColumnName)
 {
-	if (ColumnName == ConcertBrowserUtils::IconColName)
-	{
-		return GenerateIconColumn();
-	}
-
-	TSharedPtr<FConcertSessionItem> ItemPin = Item.Pin();
-	const bool bIsActiveSession = ItemPin->Type == FConcertSessionItem::EType::ActiveSession;
+	TSharedPtr<FConcertSessionTreeItem> ItemPin = Item.Pin();
+	const bool bIsActiveSession = ItemPin->Type == FConcertSessionTreeItem::EType::ActiveSession;
 	const bool bIsDefaultConfig = IsDefaultSession(ItemPin);
 	FSlateFontInfo FontInfo = GetFontInfo(bIsActiveSession, bIsDefaultConfig);
 	FSlateColor FontColor = GetFontColor(bIsActiveSession, bIsDefaultConfig);
 
 	if (ColumnName == ConcertBrowserUtils::SessionColName)
 	{
-		return GenerateSessionColumn(FontInfo, FontColor);
+		const FSlateBrush* Glyph = ItemPin->Type == FConcertSessionTreeItem::EType::ActiveSession
+			? FConcertFrontendStyle::Get()->GetBrush("Concert.ActiveSession.Icon")
+			: FConcertFrontendStyle::Get()->GetBrush("Concert.ArchivedSession.Icon");
+		return SNew( SHorizontalBox )
+
+			+SHorizontalBox::Slot()
+			.AutoWidth()
+			.Padding(6, 0, 0, 0)
+			[
+				SNew( SExpanderArrow, SharedThis(this) ).IndentAmount(12)
+			]
+
+			+SHorizontalBox::Slot()
+			.AutoWidth()
+			.Padding(0.f, 0.f, 4.f, 0.f)
+			[
+				SNew(SImage)
+				.Image(Glyph)
+			]
+
+			+SHorizontalBox::Slot()
+			.FillWidth(1.0f)
+			[
+				GenerateSessionColumn(FontInfo, FontColor)
+			];
 	}
 
 	if (ColumnName == ConcertBrowserUtils::ServerColName)
@@ -248,7 +242,7 @@ TSharedRef<SWidget> SSessionRow::GenerateWidgetForColumn(const FName& ColumnName
 
 FReply SSessionRow::OnMouseButtonDoubleClick(const FGeometry& InMyGeometry, const FPointerEvent& InMouseEvent)
 {
-	if (TSharedPtr<FConcertSessionItem> ItemPin = Item.Pin())
+	if (TSharedPtr<FConcertSessionTreeItem> ItemPin = Item.Pin())
 	{
 		DoubleClickFunc(ItemPin);
 	}
@@ -263,7 +257,7 @@ bool SSessionRow::OnValidatingSessionName(const FText& NewSessionName, FText& Ou
 
 void SSessionRow::OnSessionNameCommitted(const FText& NewSessionName, ETextCommit::Type CommitType)
 {
-	if (TSharedPtr<FConcertSessionItem> ItemPin = Item.Pin())
+	if (TSharedPtr<FConcertSessionTreeItem> ItemPin = Item.Pin())
 	{
 		FString NewName = NewSessionName.ToString();
 		if (NewName != ItemPin->SessionName) // Was renamed?
diff --git a/Engine/Plugins/Developer/Concert/ConcertUI/ConcertSharedSlate/Source/ConcertSharedSlate/Private/Session/Browser/SSessionRow.h b/Engine/Plugins/Developer/Concert/ConcertUI/ConcertSharedSlate/Source/ConcertSharedSlate/Private/Session/Browser/SSessionRow.h
--- a/Engine/Plugins/Developer/Concert/ConcertUI/ConcertSharedSlate/Source/ConcertSharedSlate/Private/Session/Browser/SSessionRow.h
+++ b/Engine/Plugins/Developer/Concert/ConcertUI/ConcertSharedSlate/Source/ConcertSharedSlate/Private/Session/Browser/SSessionRow.h
@@ -7,18 +7,18 @@
 #include "Widgets/DeclarativeSyntaxSupport.h"
 #include "Widgets/Text/SInlineEditableTextBlock.h"
 #include "Widgets/Views/STableRow.h"
-#include "Widgets/Views/SListView.h"
+#include "Widgets/Views/STreeView.h"
 
-class FConcertSessionItem;
+class FConcertSessionTreeItem;
 class SWidget;
 class SInlineEditableTextBlock;
 
-class SSessionRow : public SMultiColumnTableRow<TSharedPtr<FConcertSessionItem>>
+class SSessionRow : public SMultiColumnTableRow<TSharedPtr<FConcertSessionTreeItem>>
 {
 public:
-	using FDoubleClickFunc = TFunction<void(TSharedPtr<FConcertSessionItem>)>;
-	using FRenameFunc = TFunction<void(TSharedPtr<FConcertSessionItem>, const FString&)>;
-	using FIsDefaultSession = TFunction<bool(TSharedPtr<FConcertSessionItem>)>;
+	using FDoubleClickFunc = TFunction<void(TSharedPtr<FConcertSessionTreeItem>)>;
+	using FRenameFunc = TFunction<void(TSharedPtr<FConcertSessionTreeItem>, const FString&)>;
+	using FIsDefaultSession = TFunction<bool(TSharedPtr<FConcertSessionTreeItem>)>;
 
 	SLATE_BEGIN_ARGS(SSessionRow)
 		: _OnDoubleClickFunc()
@@ -27,16 +27,17 @@ public:
 		, _IsSelected(false)
 	{}
 
-	SLATE_ARGUMENT(FDoubleClickFunc, OnDoubleClickFunc)
-	SLATE_ARGUMENT(FRenameFunc, OnRenameFunc)
-	SLATE_ARGUMENT(FIsDefaultSession, IsDefaultSession)
-	
-	SLATE_ATTRIBUTE(FText, HighlightText)
+		SLATE_ARGUMENT(FDoubleClickFunc, OnDoubleClickFunc)
+		SLATE_ARGUMENT(FRenameFunc, OnRenameFunc)
+		SLATE_ARGUMENT(FIsDefaultSession, IsDefaultSession)
+		
+		SLATE_ATTRIBUTE(FText, HighlightText)
+		SLATE_ATTRIBUTE(FMargin, Padding)
 	SLATE_ATTRIBUTE(bool, IsSelected)
 
 	SLATE_END_ARGS()
 
-	void Construct(const FArguments& InArgs, TSharedPtr<FConcertSessionItem> InItem, const TSharedRef<STableViewBase>& InOwnerTableView);
+	void Construct(const FArguments& InArgs, TSharedPtr<FConcertSessionTreeItem> InItem, const TSharedRef<STableViewBase>& InOwnerTableView);
 	
 	virtual TSharedRef<SWidget> GenerateWidgetForColumn(const FName& ColumnName) override;
 	virtual FReply OnMouseButtonDoubleClick( const FGeometry& InMyGeometry, const FPointerEvent& InMouseEvent) override;
@@ -47,7 +48,6 @@ private:
 	FSlateFontInfo GetFontInfo(bool bIsActiveSession, bool bIsDefault);
 	FSlateColor    GetFontColor(bool bIsActiveSession, bool bIsDefault);
 
-	TSharedRef<SWidget> GenerateIconColumn();
 	TSharedRef<SWidget> GenerateSessionColumn(const FSlateFontInfo& FontInfo, const FSlateColor& FontColor);
 	TSharedRef<SWidget> GenerateServerColumn(const FSlateFontInfo& FontInfo, const FSlateColor& FontColor);
 	TSharedRef<SWidget> GenerateServerDefaultColumn(const FSlateFontInfo& FontInfo, const FSlateColor& FontColor);
@@ -58,7 +58,7 @@ private:
 	void OnBeginEditingSessionName() { SessionNameText->EnterEditingMode(); }
 	bool OnValidatingSessionName(const FText& NewSessionName, FText& OutError);
 
-	TWeakPtr<FConcertSessionItem> Item;
+	TWeakPtr<FConcertSessionTreeItem> Item;
 
 	/** Invoked when the user double click on the row */
 	FDoubleClickFunc DoubleClickFunc;
diff --git a/Engine/Plugins/Developer/Concert/ConcertUI/ConcertSharedSlate/Source/ConcertSharedSlate/Public/Session/Browser/ConcertBrowserUtils.h b/Engine/Plugins/Developer/Concert/ConcertUI/ConcertSharedSlate/Source/ConcertSharedSlate/Public/Session/Browser/ConcertBrowserUtils.h
--- a/Engine/Plugins/Developer/Concert/ConcertUI/ConcertSharedSlate/Source/ConcertSharedSlate/Public/Session/Browser/ConcertBrowserUtils.h
+++ b/Engine/Plugins/Developer/Concert/ConcertUI/ConcertSharedSlate/Source/ConcertSharedSlate/Public/Session/Browser/ConcertBrowserUtils.h
@@ -6,7 +6,7 @@
 
 #include "ConcertMessageData.h"
 #include "IConcertSessionBrowserController.h"
-#include "Session/Browser/ConcertSessionItem.h"
+#include "Session/Browser/Items/ConcertSessionTreeItem.h"
 
 #include "Internationalization/Regex.h"
 #include "Layout/Visibility.h"
@@ -19,12 +19,10 @@
 #include "Widgets/Images/SImage.h"
 #include "Widgets/Input/SButton.h"
 #include "Widgets/Layout/SBorder.h"
-#include "Widgets/Text/STextBlock.h"
 
 namespace ConcertBrowserUtils
 {
 	// Defines the sessions list view column tag names.
-	static const FName IconColName(TEXT("Icon"));
 	static const FName SessionColName(TEXT("Session"));
 	static const FName ServerColName(TEXT("Server"));
 	static const FName ProjectColName(TEXT("Project"));
@@ -103,13 +101,13 @@ namespace ConcertBrowserUtils
 	}
 
 	/** Helper function to split an array of FConcertSessionItems into multiple requests and execute*/
-	inline void RequestItemDeletion(IConcertSessionBrowserController& Controller, const TArray<TSharedPtr<FConcertSessionItem>>& SessionItems)
+	inline void RequestItemDeletion(IConcertSessionBrowserController& Controller, const TArray<TSharedPtr<FConcertSessionTreeItem>>& SessionItems)
 	{
 		using FServerAdminEndpointId = FGuid;
 		using FSessionId = FGuid;
 			
 		TMap<FServerAdminEndpointId, TArray<FSessionId>> Requests;
-		for (const TSharedPtr<FConcertSessionItem>& Item : SessionItems)
+		for (const TSharedPtr<FConcertSessionTreeItem>& Item : SessionItems)
 		{
 			Requests.FindOrAdd(Item->ServerAdminEndpointId).Add(Item->SessionId);
 		}
diff --git a/Engine/Plugins/Developer/Concert/ConcertUI/ConcertSharedSlate/Source/ConcertSharedSlate/Public/Session/Browser/ConcertSessionItem.h b/Engine/Plugins/Developer/Concert/ConcertUI/ConcertSharedSlate/Source/ConcertSharedSlate/Public/Session/Browser/ConcertSessionItem.h
deleted file mode 100644
diff --git a/Engine/Plugins/Developer/Concert/ConcertUI/ConcertSharedSlate/Source/ConcertSharedSlate/Public/Session/Browser/Items/ConcertActiveGroupTreeItem.h b/Engine/Plugins/Developer/Concert/ConcertUI/ConcertSharedSlate/Source/ConcertSharedSlate/Public/Session/Browser/Items/ConcertActiveGroupTreeItem.h
new file mode 100644
--- /dev/null
+++ b/Engine/Plugins/Developer/Concert/ConcertUI/ConcertSharedSlate/Source/ConcertSharedSlate/Public/Session/Browser/Items/ConcertActiveGroupTreeItem.h
@@ -0,0 +1,18 @@
+// Copyright Epic Games, Inc. All Rights Reserved.
+
+#pragma once
+
+#include "CoreMinimal.h"
+#include "ConcertGroupTreeItem.h"
+#include "ConcertSessionTreeItem.h"
+
+/** Groups several active sessions */
+class CONCERTSHAREDSLATE_API FConcertActiveGroupTreeItem : public FConcertGroupTreeItem
+{
+	using Super = FConcertGroupTreeItem;
+public:
+	
+	FConcertActiveGroupTreeItem(FGetSessions GetSessionsFunc)
+		: Super(MoveTemp(GetSessionsFunc), { FConcertSessionTreeItem::EType::ActiveSession, FConcertSessionTreeItem::EType::NewSession, FConcertSessionTreeItem::EType::SaveSession })
+	{}
+};
\ No newline at end of file
diff --git a/Engine/Plugins/Developer/Concert/ConcertUI/ConcertSharedSlate/Source/ConcertSharedSlate/Public/Session/Browser/Items/ConcertArchivedGroupTreeItem.h b/Engine/Plugins/Developer/Concert/ConcertUI/ConcertSharedSlate/Source/ConcertSharedSlate/Public/Session/Browser/Items/ConcertArchivedGroupTreeItem.h
new file mode 100644
--- /dev/null
+++ b/Engine/Plugins/Developer/Concert/ConcertUI/ConcertSharedSlate/Source/ConcertSharedSlate/Public/Session/Browser/Items/ConcertArchivedGroupTreeItem.h
@@ -0,0 +1,18 @@
+
+// Copyright Epic Games, Inc. All Rights Reserved.
+
+#pragma once
+
+#include "CoreMinimal.h"
+#include "ConcertGroupTreeItem.h"
+
+/** Groups archived sessions */
+class CONCERTSHAREDSLATE_API FConcertArchivedGroupTreeItem : public FConcertGroupTreeItem
+{
+	using Super = FConcertGroupTreeItem;
+public:
+	
+	FConcertArchivedGroupTreeItem(FGetSessions GetSessionsFunc)
+		: Super(MoveTemp(GetSessionsFunc), { FConcertSessionTreeItem::EType::ArchivedSession, FConcertSessionTreeItem::EType::RestoreSession })
+	{}
+};
\ No newline at end of file
diff --git a/Engine/Plugins/Developer/Concert/ConcertUI/ConcertSharedSlate/Source/ConcertSharedSlate/Public/Session/Browser/Items/ConcertGroupTreeItem.h b/Engine/Plugins/Developer/Concert/ConcertUI/ConcertSharedSlate/Source/ConcertSharedSlate/Public/Session/Browser/Items/ConcertGroupTreeItem.h
new file mode 100644
--- /dev/null
+++ b/Engine/Plugins/Developer/Concert/ConcertUI/ConcertSharedSlate/Source/ConcertSharedSlate/Public/Session/Browser/Items/ConcertGroupTreeItem.h
@@ -0,0 +1,38 @@
+
+// Copyright Epic Games, Inc. All Rights Reserved.
+
+#pragma once
+
+#include "CoreMinimal.h"
+#include "ConcertSessionTreeItem.h"
+#include "ConcertTreeItem.h"
+
+/** Groups several FConcertSessionTreeItem */
+class CONCERTSHAREDSLATE_API FConcertGroupTreeItem : public FConcertTreeItem
+{
+public:
+
+	using FGetSessions = TFunction<TArray<TSharedPtr<FConcertSessionTreeItem>>()>;
+
+	FConcertGroupTreeItem(FGetSessions GetSessionsFunc, TSet<FConcertSessionTreeItem::EType> AllowedTypes)
+		: GetSessionsFunc(MoveTemp(GetSessionsFunc))
+		, AllowedTypes(MoveTemp(AllowedTypes))
+	{}
+
+	virtual void GetChildren(TArray<TSharedPtr<FConcertTreeItem>>& OutChildren) const override
+	{
+		const TArray<TSharedPtr<FConcertSessionTreeItem>>& Items = GetSessionsFunc(); 
+		for (const TSharedPtr<FConcertSessionTreeItem>& Item : Items)
+		{
+			if (AllowedTypes.Contains(Item->Type))
+			{
+				OutChildren.Add(Item);
+			}
+		}
+	}
+
+protected:
+	
+	const FGetSessions GetSessionsFunc;
+	const TSet<FConcertSessionTreeItem::EType> AllowedTypes;
+};
\ No newline at end of file
diff --git a/Engine/Plugins/Developer/Concert/ConcertUI/ConcertSharedSlate/Source/ConcertSharedSlate/Public/Session/Browser/Items/ConcertSessionTreeItem.h b/Engine/Plugins/Developer/Concert/ConcertUI/ConcertSharedSlate/Source/ConcertSharedSlate/Public/Session/Browser/Items/ConcertSessionTreeItem.h
new file mode 100644
--- /dev/null
+++ b/Engine/Plugins/Developer/Concert/ConcertUI/ConcertSharedSlate/Source/ConcertSharedSlate/Public/Session/Browser/Items/ConcertSessionTreeItem.h
@@ -0,0 +1,77 @@
+﻿// Copyright Epic Games, Inc. All Rights Reserved.
+
+#pragma once
+
+#include "CoreMinimal.h"
+#include "ConcertMessageData.h"
+#include "ConcertTreeItem.h"
+
+/** Signal emitted when a session name text field should enter in edit mode. */
+DECLARE_MULTICAST_DELEGATE(FOnBeginEditConcertSessionNameRequest)
+
+/** An item that represents a session, i.e. has session info. */
+class CONCERTSHAREDSLATE_API FConcertSessionTreeItem : public FConcertTreeItem
+{
+public:
+	enum class EType
+	{
+		None,
+		NewSession,				// Editable item to enter a session name and a pick a server.
+		RestoreSession,			// Editable item to name the restored session.
+		SaveSession,			// Editable item to name the archive.
+		ActiveSession,			// Read-only item representing an active session.
+		ArchivedSession,		// Read-only item representing an archived session.
+	};
+
+	FConcertSessionTreeItem(EType Type = EType::None,
+		const FGuid& ServerAdminEndpointId = {},
+		const FGuid& SessionId = {},
+		const FString& SessionName = {},
+		const FString& ServerName = {},
+		const FString& ProjectName = {},
+		const FString& ProjectVersion = {},
+		EConcertServerFlags ServerFlags = {},
+		const FDateTime& LastModified = {},
+		const FOnBeginEditConcertSessionNameRequest& OnBeginEditSessionNameRequest = {}
+		)
+		: Type(Type)
+		, ServerAdminEndpointId(ServerAdminEndpointId)
+		, SessionId(SessionId)
+		, SessionName(SessionName)
+		, ServerName(ServerName)
+		, ProjectName(ProjectName)
+		, ProjectVersion(ProjectVersion)
+		, ServerFlags(ServerFlags)
+		, LastModified(LastModified)
+		, OnBeginEditSessionNameRequest(OnBeginEditSessionNameRequest)
+	{}
+
+	bool operator==(const FConcertSessionTreeItem& Other) const
+	{
+		return Type == Other.Type && ServerAdminEndpointId == Other.ServerAdminEndpointId && SessionId == Other.SessionId;
+	}
+
+	FConcertSessionTreeItem MakeCopyAsType(EType NewType) const
+	{
+		FConcertSessionTreeItem NewItem = *this;
+		NewItem.Type = NewType;
+		return NewItem;
+	}
+
+	virtual void GetChildren(TArray<TSharedPtr<FConcertTreeItem>>& OutChildren) const override {}
+
+	EType Type = EType::None;
+	FGuid ServerAdminEndpointId;
+	FGuid SessionId;
+	FString SessionName;
+	FString ServerName;
+	FString ProjectName;
+	FString ProjectVersion;
+	EConcertServerFlags ServerFlags = EConcertServerFlags::None;
+	
+	FDateTime LastModified;
+
+	FOnBeginEditConcertSessionNameRequest OnBeginEditSessionNameRequest; // Emitted when user press 'F2' or select 'Rename' from context menu.
+};
+
+
diff --git a/Engine/Plugins/Developer/Concert/ConcertUI/ConcertSharedSlate/Source/ConcertSharedSlate/Public/Session/Browser/Items/ConcertTreeItem.h b/Engine/Plugins/Developer/Concert/ConcertUI/ConcertSharedSlate/Source/ConcertSharedSlate/Public/Session/Browser/Items/ConcertTreeItem.h
new file mode 100644
--- /dev/null
+++ b/Engine/Plugins/Developer/Concert/ConcertUI/ConcertSharedSlate/Source/ConcertSharedSlate/Public/Session/Browser/Items/ConcertTreeItem.h
@@ -0,0 +1,15 @@
+// Copyright Epic Games, Inc. All Rights Reserved.
+
+#pragma once
+
+#include "CoreMinimal.h"
+
+/** Base class for everything displayed in session view */
+class CONCERTSHAREDSLATE_API FConcertTreeItem
+{
+public:
+
+	virtual void GetChildren(TArray<TSharedPtr<FConcertTreeItem>>& OutChildren) const = 0;
+	
+	virtual ~FConcertTreeItem() = default;
+};
\ No newline at end of file
diff --git a/Engine/Plugins/Developer/Concert/ConcertUI/ConcertSharedSlate/Source/ConcertSharedSlate/Public/Session/Browser/SConcertSessionBrowser.h b/Engine/Plugins/Developer/Concert/ConcertUI/ConcertSharedSlate/Source/ConcertSharedSlate/Public/Session/Browser/SConcertSessionBrowser.h
--- a/Engine/Plugins/Developer/Concert/ConcertUI/ConcertSharedSlate/Source/ConcertSharedSlate/Public/Session/Browser/SConcertSessionBrowser.h
+++ b/Engine/Plugins/Developer/Concert/ConcertUI/ConcertSharedSlate/Source/ConcertSharedSlate/Public/Session/Browser/SConcertSessionBrowser.h
@@ -10,9 +10,12 @@
 #include "Widgets/DeclarativeSyntaxSupport.h"
 #include "Widgets/SCompoundWidget.h"
 #include "Widgets/Views/SHeaderRow.h"
-#include "Widgets/Views/SListView.h"
+#include "Widgets/Views/STreeView.h"
 
-class FConcertSessionItem;
+class FConcertActiveGroupTreeItem;
+class FConcertArchivedGroupTreeItem;
+class FConcertTreeItem;
+class FConcertSessionTreeItem;
 class FExtender;
 class IConcertSessionBrowserController;
 class ITableRow;
@@ -25,10 +28,10 @@ struct FConcertSessionInfo;
 
 DECLARE_DELEGATE_RetVal_OneParam(TSharedRef<SWidget>, FExtendSessionTable, const TSharedRef<SWidget>& /* TableView */);
 DECLARE_DELEGATE_OneParam(FExtenderDelegate, FExtender&);
-DECLARE_DELEGATE_TwoParams(FExtendSessionContextMenu, const TSharedPtr<FConcertSessionItem>&, FExtender&)
-DECLARE_DELEGATE_OneParam(FSessionDelegate, const TSharedPtr<FConcertSessionItem>& /*Session*/);
-DECLARE_DELEGATE_OneParam(FSessionListDelegate, const TArray<TSharedPtr<FConcertSessionItem>>& /*Sessions*/);
-DECLARE_DELEGATE_RetVal_OneParam(bool, FCanRemoveSessions, const TArray<TSharedPtr<FConcertSessionItem>>& /* SessionItem */);
+DECLARE_DELEGATE_TwoParams(FExtendSessionContextMenu, const TSharedPtr<FConcertSessionTreeItem>&, FExtender&)
+DECLARE_DELEGATE_OneParam(FSessionDelegate, const TSharedPtr<FConcertSessionTreeItem>& /*Session*/);
+DECLARE_DELEGATE_OneParam(FSessionListDelegate, const TArray<TSharedPtr<FConcertSessionTreeItem>>& /*Sessions*/);
+DECLARE_DELEGATE_RetVal_OneParam(bool, FCanRemoveSessions, const TArray<TSharedPtr<FConcertSessionTreeItem>>& /* SessionItem */);
 
 /**
  * Enables the user to browse/search/filter/sort active and archived sessions, create new session,
@@ -104,7 +107,7 @@ public:
 	
 	void RefreshSessionList();
 	bool HasAnySessions() const { return Sessions.Num() > 0; }
-	TArray<TSharedPtr<FConcertSessionItem>> GetSelectedItems() const { return SessionsView->GetSelectedItems(); }
+	TArray<TSharedPtr<FConcertSessionTreeItem>> GetSelectedItems() const;
 
 	bool IsNewButtonEnabled() const { return IsNewButtonEnabledInternal(); }
 	bool IsRestoreButtonEnabled() const { return IsRestoreButtonEnabledInternal(); }
@@ -115,26 +118,35 @@ public:
 	/** Adds row for creating new session. Exposed for other widgets, e.g. discovery overlay to create a new session. */
 	void InsertNewSessionEditableRow() { InsertNewSessionEditableRowInternal(); }
 	/** Creates row under the given (archived) session with which session can be restored. */
-	void InsertRestoreSessionAsEditableRow(const TSharedPtr<FConcertSessionItem>& ArchivedItem) { InsertRestoreSessionAsEditableRowInternal(ArchivedItem); }
+	void InsertRestoreSessionAsEditableRow(const TSharedPtr<FConcertSessionTreeItem>& ArchivedItem) { InsertRestoreSessionAsEditableRowInternal(ArchivedItem); }
 	
 private:
+
+	enum class EInsertPosition
+	{
+		AtBeginning,
+		AfterParent
+	};
 	
 	// Layout the 'session|details' split view.
 	TSharedRef<SWidget> MakeBrowserContent(const FArguments& InArgs);
+	
+	void OnGetChildren(TSharedPtr<FConcertTreeItem> InItem, TArray<TSharedPtr<FConcertTreeItem>>& OutChildren);
 
 	// Layout the sessions view and controls.
 	TSharedRef<SWidget> MakeControlBar(const FArguments& InArgs);
 	TSharedRef<SWidget> MakeButtonBar(const FArguments& InArgs);
 	TSharedRef<SWidget> MakeSessionTableView(const FArguments& InArgs);
 	TSharedRef<SWidget> MakeSessionTableFooter();
-	TSharedRef<ITableRow> OnGenerateSessionRowWidget(TSharedPtr<FConcertSessionItem> Item, const TSharedRef<STableViewBase>& OwnerTable);
+	TSharedRef<ITableRow> OnGenerateSessionRowWidget(TSharedPtr<FConcertTreeItem> Item, const TSharedRef<STableViewBase>& OwnerTable);
 
 	// Creates row widgets for session list view, validates user inputs and forward user requests for processing to a delegate function implemented by this class.
-	TSharedRef<ITableRow> MakeActiveSessionRowWidget(const TSharedPtr<FConcertSessionItem>& ActiveItem, const TSharedRef<STableViewBase>& OwnerTable);
-	TSharedRef<ITableRow> MakeArchivedSessionRowWidget(const TSharedPtr<FConcertSessionItem>& ArchivedItem, const TSharedRef<STableViewBase>& OwnerTable);
-	TSharedRef<ITableRow> MakeNewSessionRowWidget(const TSharedPtr<FConcertSessionItem>& NewItem, const TSharedRef<STableViewBase>& OwnerTable);
-	TSharedRef<ITableRow> MakeRestoreSessionRowWidget(const TSharedPtr<FConcertSessionItem>& RestoreItem, const TSharedRef<STableViewBase>& OwnerTable);
-	TSharedRef<ITableRow> MakeSaveSessionRowWidget(const TSharedPtr<FConcertSessionItem>& ArchivedItem, const TSharedRef<STableViewBase>& OwnerTable);
+	TSharedRef<ITableRow> MakeActiveSessionRowWidget(const TSharedPtr<FConcertSessionTreeItem>& ActiveItem, const TSharedRef<STableViewBase>& OwnerTable);
+	TSharedRef<ITableRow> MakeArchivedSessionRowWidget(const TSharedPtr<FConcertSessionTreeItem>& ArchivedItem, const TSharedRef<STableViewBase>& OwnerTable);
+	TSharedRef<ITableRow> MakeNewSessionRowWidget(const TSharedPtr<FConcertSessionTreeItem>& NewItem, const TSharedRef<STableViewBase>& OwnerTable);
+	TSharedRef<ITableRow> MakeRestoreSessionRowWidget(const TSharedPtr<FConcertSessionTreeItem>& RestoreItem, const TSharedRef<STableViewBase>& OwnerTable);
+	TSharedRef<ITableRow> MakeSaveSessionRowWidget(const TSharedPtr<FConcertSessionTreeItem>& ArchivedItem, const TSharedRef<STableViewBase>& OwnerTable);
+	TSharedRef<ITableRow> MakeGroupRowWidget(const TSharedPtr<FConcertTreeItem>& GroupItem, const TSharedRef<STableViewBase>& OwnerTable); 
 
 	// Creates the contextual menu when right clicking a session list view row.
 	TSharedPtr<SWidget> MakeTableContextualMenu();
@@ -150,40 +162,51 @@ private:
 	FReply OnRestoreButtonClicked();
 	FReply OnArchiveButtonClicked();
 	FReply OnDeleteButtonClicked();
-	void OnBeginEditingSessionName(TSharedPtr<FConcertSessionItem> Item);
+	void OnBeginEditingSessionName(TSharedPtr<FConcertSessionTreeItem> Item);
 
 	// Manipulates the sessions view (the array and the UI).
-	void OnSessionSelectionChanged(TSharedPtr<FConcertSessionItem> SelectedSession, ESelectInfo::Type SelectInfo);
+	void OnSessionSelectionChanged(TSharedPtr<FConcertTreeItem> SelectedSession, ESelectInfo::Type SelectInfo);
+	void OnSessionSelectionChanged(TSharedPtr<FConcertSessionTreeItem> SelectedSession, ESelectInfo::Type SelectInfo);
 	void InsertNewSessionEditableRowInternal();
-	void InsertRestoreSessionAsEditableRowInternal(const TSharedPtr<FConcertSessionItem>& ArchivedItem);
-	void InsertArchiveSessionAsEditableRow(const TSharedPtr<FConcertSessionItem>& ActiveItem);
-	void InsertEditableSessionRow(TSharedPtr<FConcertSessionItem> EditableItem, TSharedPtr<FConcertSessionItem> ParentItem);
-	void RemoveSessionRow(const TSharedPtr<FConcertSessionItem>& Item);
+	void InsertRestoreSessionAsEditableRowInternal(const TSharedPtr<FConcertSessionTreeItem>& ArchivedItem);
+	void InsertArchiveSessionAsEditableRow(const TSharedPtr<FConcertSessionTreeItem>& ActiveItem);
+	void InsertEditableSessionRow(TSharedPtr<FConcertSessionTreeItem> EditableItem, TSharedPtr<FConcertSessionTreeItem> ParentItem);
+	void RemoveSessionRow(const TSharedPtr<FConcertSessionTreeItem>& Item);
 
 	// Sessions sorting. (Sorts the session view)
 	EColumnSortMode::Type GetColumnSortMode(const FName ColumnId) const;
 	EColumnSortPriority::Type GetColumnSortPriority(const FName ColumnId) const;
 	void OnColumnSortModeChanged(const EColumnSortPriority::Type SortPriority, const FName& ColumnId, const EColumnSortMode::Type InSortMode);
-	void SortSessionList();
 	void EnsureEditableParentChildOrder();
 
 	// Sessions filtering. (Filters the session view)
 	void OnSearchTextChanged(const FText& InFilterText);
 	void OnSearchTextCommitted(const FText& InFilterText, ETextCommit::Type CommitType);
 	void OnFilterMenuChecked(const FName MenuName);
-	void PopulateSearchStrings(const FConcertSessionItem& Item, TArray<FString>& OutSearchStrings) const;
-	bool IsFilteredOut(const FConcertSessionItem& Item) const;
+	void PopulateSearchStrings(const FConcertSessionTreeItem& Item, TArray<FString>& OutSearchStrings) const;
+	bool IsFilteredOut(const FConcertSessionTreeItem& Item) const;
 	FText HighlightSearchText() const;
 
 	// Passes the user requests to FConcertBrowserController.
-	void RequestCreateSession(const TSharedPtr<FConcertSessionItem>& NewItem);
-	void RequestArchiveSession(const TSharedPtr<FConcertSessionItem>& ActiveItem, const FString& ArchiveName);
-	void RequestRestoreSession(const TSharedPtr<FConcertSessionItem>& RestoreItem, const FString& SessionName);
-	void RequestRenameSession(const TSharedPtr<FConcertSessionItem>& RenamedItem, const FString& NewName);
-	void RequestDeleteSessions(const TArray<TSharedPtr<FConcertSessionItem>>& ItemsToDelete);
+	void RequestCreateSession(const TSharedPtr<FConcertSessionTreeItem>& NewItem);
+	void RequestArchiveSession(const TSharedPtr<FConcertSessionTreeItem>& ActiveItem, const FString& ArchiveName);
+	void RequestRestoreSession(const TSharedPtr<FConcertSessionTreeItem>& RestoreItem, const FString& SessionName);
+	void RequestRenameSession(const TSharedPtr<FConcertSessionTreeItem>& RenamedItem, const FString& NewName);
+	void RequestDeleteSessions(const TArray<TSharedPtr<FConcertSessionTreeItem>>& ItemsToDelete);
 
 	TSharedPtr<IConcertSessionBrowserController> GetController() const;
 
+	bool IsGroupItem(const TSharedPtr<FConcertTreeItem>& TreeItem) const;
+
+	// Modifying sessions arrays
+	void ResetSessions();
+	void AddSession(TSharedPtr<FConcertSessionTreeItem> Session);
+	void InsertSessionAtBeginning(TSharedPtr<FConcertSessionTreeItem> NewSession);
+	void InsertSessionAfterParent(TSharedPtr<FConcertSessionTreeItem> NewSession, TSharedPtr<FConcertSessionTreeItem> Parent);
+	EInsertPosition InsertSessionAfterParentIfAvailableOrAtBeginning(TSharedPtr<FConcertSessionTreeItem> NewSession, TSharedPtr<FConcertSessionTreeItem> Parent);
+	void RemoveSession(TSharedPtr<FConcertSessionTreeItem> Session);
+	void SortSessionList();
+
 private:
 	
 	// Gives access to the concert data (servers, sessions, clients, etc).
@@ -207,16 +230,22 @@ private:
 	FSessionListDelegate PostRequestedDeleteSession;
 	FCanRemoveSessions AskUserToDeleteSessions;
 
-	// The items displayed in the session list view. It might be filtered and sorted compared to the full list hold by the controller.
-	TArray<TSharedPtr<FConcertSessionItem>> Sessions;
+	/** The items displayed in the session list view. It might be filtered and sorted compared to the full list hold by the controller. */
+	TArray<TSharedPtr<FConcertSessionTreeItem>> Sessions;
+	/** Groups all active sessions under it */
+	TSharedPtr<FConcertActiveGroupTreeItem> ActiveSessionsGroupItem;
+	/** Groups all archived sessions under it */
+	TSharedPtr<FConcertArchivedGroupTreeItem> ArchivedSessionsGroupItem;
+	/** ItemSource for SessionsView containing Sessions, ActiveGroup, and ArchivedGroup */
+	TArray<TSharedPtr<FConcertTreeItem>> TreeItemSource;
 
 	// The session list view.
-	TSharedPtr<SListView<TSharedPtr<FConcertSessionItem>>> SessionsView;
+	TSharedPtr<STreeView<TSharedPtr<FConcertTreeItem>>> SessionsView;
 	TSharedPtr<SHeaderRow> SessionHeaderRow;
 
 	// The item corresponding to a row used to create/archive/restore a session. There is only one at the time
-	TSharedPtr<FConcertSessionItem> EditableSessionRow;
-	TSharedPtr<FConcertSessionItem> EditableSessionRowParent; // For archive/restore, indicate which element is archived or restored.
+	TSharedPtr<FConcertSessionTreeItem> EditableSessionRow;
+	TSharedPtr<FConcertSessionTreeItem> EditableSessionRowParent; // For archive/restore, indicate which element is archived or restored.
 
 	// Sorting.
 	EColumnSortMode::Type PrimarySortMode = EColumnSortMode::None;
@@ -226,7 +255,7 @@ private:
 
 	// Filtering.
 	TSharedPtr<SSearchBox> SearchBox;
-	TSharedPtr<TTextFilter<const FConcertSessionItem&>> SearchTextFilter;
+	TSharedPtr<TTextFilter<const FConcertSessionTreeItem&>> SearchTextFilter;
 	TSharedPtr<FText> SearchedText;
 	bool bRefreshSessionFilter = true;
 	FString LastDefaultServerUrl;
