Archive session context menu now contains delete option

diff --git a/Engine/Plugins/Developer/Concert/ConcertUI/ConcertSharedSlate/Source/ConcertSharedSlate/Private/Session/Activity/SConcertSessionActivities.cpp b/Engine/Plugins/Developer/Concert/ConcertUI/ConcertSharedSlate/Source/ConcertSharedSlate/Private/Session/Activity/SConcertSessionActivities.cpp
--- a/Engine/Plugins/Developer/Concert/ConcertUI/ConcertSharedSlate/Source/ConcertSharedSlate/Private/Session/Activity/SConcertSessionActivities.cpp
+++ b/Engine/Plugins/Developer/Concert/ConcertUI/ConcertSharedSlate/Source/ConcertSharedSlate/Private/Session/Activity/SConcertSessionActivities.cpp
@@ -223,7 +223,7 @@ void SConcertSessionActivities::Construct(const FArguments& InArgs)
 	GetTransactionEventFn = InArgs._OnGetTransactionEvent;
 	GetPackageEventFn = InArgs._OnGetPackageEvent;
 	MakeColumnOverlayWidgetFn = InArgs._OnMakeColumnOverlayWidget;
-
+	OnContextMenuOpening = InArgs._OnContextMenuOpening;
 	
 	HighlightText = InArgs._HighlightText;
 	
@@ -275,7 +275,19 @@ void SConcertSessionActivities::Construct(const FArguments& InArgs)
 					.OnListViewScrolled(this, &SConcertSessionActivities::OnListViewScrolled)
 					.OnSelectionChanged(this, &SConcertSessionActivities::OnListViewSelectionChanged)
 					.HeaderRow(HeaderRow)
-					.OnContextMenuOpening_Lambda([this](){ return UE::ConcertSharedSlate::MakeTableContextMenu(HeaderRow.ToSharedRef(), {}, true); })
+					.OnContextMenuOpening_Lambda([this]() -> TSharedPtr<SWidget>
+					{
+						if (ActivityView->GetSelectedItems().Num() == 0)
+						{
+							return UE::ConcertSharedSlate::MakeTableContextMenu(HeaderRow.ToSharedRef(), {}, true);
+						}
+
+						if (OnContextMenuOpening.IsBound())
+						{
+							return OnContextMenuOpening.Execute();
+						}
+						return {};
+					})
 				]
 			]
 
diff --git a/Engine/Plugins/Developer/Concert/ConcertUI/ConcertSharedSlate/Source/ConcertSharedSlate/Private/Session/History/SEditableSessionHistory.cpp b/Engine/Plugins/Developer/Concert/ConcertUI/ConcertSharedSlate/Source/ConcertSharedSlate/Private/Session/History/SEditableSessionHistory.cpp
--- a/Engine/Plugins/Developer/Concert/ConcertUI/ConcertSharedSlate/Source/ConcertSharedSlate/Private/Session/History/SEditableSessionHistory.cpp
+++ b/Engine/Plugins/Developer/Concert/ConcertUI/ConcertSharedSlate/Source/ConcertSharedSlate/Private/Session/History/SEditableSessionHistory.cpp
@@ -3,6 +3,7 @@
 #include "Session/History/SEditableSessionHistory.h"
 
 #include "Algo/Transform.h"
+#include "Framework/MultiBox/MultiBoxBuilder.h"
 #include "Styling/AppStyle.h"
 #include "SNegativeActionButton.h"
 #include "Session/History/SSessionHistory.h"
@@ -28,22 +29,11 @@ void SEditableSessionHistory::Construct(const FArguments& InArgs)
 	SessionHistory = InArgs._MakeSessionHistory.Execute(
 		SSessionHistory::FArguments()
 		.SelectionMode(ESelectionMode::Multi)
+		.OnContextMenuOpening(this, &SEditableSessionHistory::OnContextMenuOpening)
 		.SearchButtonArea()
 		[
 			SNew(SHorizontalBox)
 
-			// Delete
-			+SHorizontalBox::Slot()
-			.AutoWidth()
-			[
-				SNew(SNegativeActionButton)
-				.ActionButtonStyle(EActionButtonStyle::Error) // Red
-				.OnClicked(this, &SEditableSessionHistory::OnClickDeleteActivitiesButton)
-				.ToolTipText(this, &SEditableSessionHistory::GetDeleteActivitiesToolTip)
-				.IsEnabled(this, &SEditableSessionHistory::IsDeleteButtonEnabled)
-				.Icon(FAppStyle::GetBrush("Icons.Delete"))
-			]
-
 			// Mute
 			+SHorizontalBox::Slot()
 			.AutoWidth()
@@ -91,6 +81,56 @@ FReply SEditableSessionHistory::OnKeyDown(const FGeometry& MyGeometry, const FKe
 	return SCompoundWidget::OnKeyDown(MyGeometry, InKeyEvent);
 }
 
+TSharedPtr<SWidget> SEditableSessionHistory::OnContextMenuOpening()
+{
+	FMenuBuilder MenuBuilder(true, nullptr);
+
+	MenuBuilder.AddMenuEntry(
+		LOCTEXT("MuteMenuLabel", "Mute"),
+		FText::GetEmpty(),
+		FSlateIcon(FConcertFrontendStyle::GetStyleSetName(), "Concert.MuteActivities"),
+		FUIAction(
+		FExecuteAction::CreateLambda([this](){ OnClickMuteActivitesButton(); }),
+			FCanExecuteAction::CreateLambda([this] { return IsMuteButtonEnabled(); })
+		),
+		NAME_None,
+		EUserInterfaceActionType::Button
+		);
+	MenuBuilder.AddMenuEntry(
+		LOCTEXT("UnmuteMenuLabel", "Unmute"),
+		FText::GetEmpty(),
+		FSlateIcon(FConcertFrontendStyle::GetStyleSetName(), "Concert.UnmuteActivities"),
+		FUIAction(
+		FExecuteAction::CreateLambda([this](){ OnClickUnmuteActivitesButton(); }),
+			FCanExecuteAction::CreateLambda([this] { return IsUnmuteButtonEnabled(); })
+		),
+		NAME_None,
+		EUserInterfaceActionType::Button
+		);
+	MenuBuilder.AddSubMenu(
+		LOCTEXT("EditMenuLabel", "Edit"),
+		FText::GetEmpty(),
+		FNewMenuDelegate::CreateLambda([this](FMenuBuilder& SubMenuBuilder)
+		{
+			SubMenuBuilder.AddMenuEntry(
+				LOCTEXT("DeleteMenuLabel", "Delete"),
+				FText::GetEmpty(),
+				FSlateIcon(FAppStyle::GetAppStyleSetName(), "Icons.Delete"),
+				FUIAction(
+				FExecuteAction::CreateLambda([this](){ OnClickDeleteActivitiesButton(); }),
+					FCanExecuteAction::CreateLambda([this] { return IsDeleteButtonEnabled(); })
+				),
+				NAME_None,
+				EUserInterfaceActionType::Button
+			);
+		}),
+		false,
+		FSlateIcon(FAppStyle::GetAppStyleSetName(), "Icons.Edit")
+	);
+	
+	return MenuBuilder.MakeWidget();
+}
+
 FReply SEditableSessionHistory::OnClickDeleteActivitiesButton() const
 {
 	if (DeleteActivitiesFunc.IsBound() && CanDeleteActivitiesFunc.Execute(SessionHistory->GetSelectedActivities()).CanPerformAction())
diff --git a/Engine/Plugins/Developer/Concert/ConcertUI/ConcertSharedSlate/Source/ConcertSharedSlate/Private/Session/History/SSessionHistory.cpp b/Engine/Plugins/Developer/Concert/ConcertUI/ConcertSharedSlate/Source/ConcertSharedSlate/Private/Session/History/SSessionHistory.cpp
--- a/Engine/Plugins/Developer/Concert/ConcertUI/ConcertSharedSlate/Source/ConcertSharedSlate/Private/Session/History/SSessionHistory.cpp
+++ b/Engine/Plugins/Developer/Concert/ConcertUI/ConcertSharedSlate/Source/ConcertSharedSlate/Private/Session/History/SSessionHistory.cpp
@@ -61,6 +61,7 @@ void SSessionHistory::Construct(const FArguments& InArgs)
 	SAssignNew(ActivityListView, SConcertSessionActivities)
 		.OnGetPackageEvent(InArgs._GetPackageEvent)
 		.OnGetTransactionEvent(InArgs._GetTransactionEvent)
+		.OnContextMenuOpening(InArgs._OnContextMenuOpening)
 		.OnMakeColumnOverlayWidget(InArgs._OnMakeColumnOverlayWidget)
 		.OnMapActivityToClient(this, &SSessionHistory::GetClientInfo)
 		.HighlightText(this, &SSessionHistory::HighlightSearchedText)
diff --git a/Engine/Plugins/Developer/Concert/ConcertUI/ConcertSharedSlate/Source/ConcertSharedSlate/Public/Session/Activity/SConcertSessionActivities.h b/Engine/Plugins/Developer/Concert/ConcertUI/ConcertSharedSlate/Source/ConcertSharedSlate/Public/Session/Activity/SConcertSessionActivities.h
--- a/Engine/Plugins/Developer/Concert/ConcertUI/ConcertSharedSlate/Source/ConcertSharedSlate/Public/Session/Activity/SConcertSessionActivities.h
+++ b/Engine/Plugins/Developer/Concert/ConcertUI/ConcertSharedSlate/Source/ConcertSharedSlate/Public/Session/Activity/SConcertSessionActivities.h
@@ -107,6 +107,9 @@ public:
 
 		/** If bound, invoked when generating a row to add an overlay to a column. */
 		SLATE_EVENT(FMakeColumnOverlayWidgetFunc, OnMakeColumnOverlayWidget)
+
+		/** If bound, invoked when an item in the table is right-click and we are supposed to show a menu of actions. */
+		SLATE_EVENT(FOnContextMenuOpening, OnContextMenuOpening)
 	
 		/** The columns to show for each activity; they are shown after date time */
 		SLATE_ARGUMENT(TArray<FActivityColumn>, Columns)
@@ -261,6 +264,9 @@ private:
 	/** Used to overlay a widget over a column widget (add an extra layer above the normal one) */
 	FMakeColumnOverlayWidgetFunc MakeColumnOverlayWidgetFn;
 
+	/** If bound, invoked when an item in the table is right-click and we are supposed to show a menu of actions. */
+	FOnContextMenuOpening OnContextMenuOpening;
+
 	/** Whether to darken muted activities */
 	bool bDarkenMutedActivities = true;
 
diff --git a/Engine/Plugins/Developer/Concert/ConcertUI/ConcertSharedSlate/Source/ConcertSharedSlate/Public/Session/History/SEditableSessionHistory.h b/Engine/Plugins/Developer/Concert/ConcertUI/ConcertSharedSlate/Source/ConcertSharedSlate/Public/Session/History/SEditableSessionHistory.h
--- a/Engine/Plugins/Developer/Concert/ConcertUI/ConcertSharedSlate/Source/ConcertSharedSlate/Public/Session/History/SEditableSessionHistory.h
+++ b/Engine/Plugins/Developer/Concert/ConcertUI/ConcertSharedSlate/Source/ConcertSharedSlate/Public/Session/History/SEditableSessionHistory.h
@@ -70,6 +70,8 @@ private:
 	FCanPerformActionOnActivities CanUnmuteActivitiesFunc;
 	FRequestActivitiesAction UnmuteActivitiesFunc;
 
+	TSharedPtr<SWidget> OnContextMenuOpening();
+	
 	FReply OnClickDeleteActivitiesButton() const;
 	FText GetDeleteActivitiesToolTip() const;
 	bool IsDeleteButtonEnabled() const;
diff --git a/Engine/Plugins/Developer/Concert/ConcertUI/ConcertSharedSlate/Source/ConcertSharedSlate/Public/Session/History/SSessionHistory.h b/Engine/Plugins/Developer/Concert/ConcertUI/ConcertSharedSlate/Source/ConcertSharedSlate/Public/Session/History/SSessionHistory.h
--- a/Engine/Plugins/Developer/Concert/ConcertUI/ConcertSharedSlate/Source/ConcertSharedSlate/Public/Session/History/SSessionHistory.h
+++ b/Engine/Plugins/Developer/Concert/ConcertUI/ConcertSharedSlate/Source/ConcertSharedSlate/Public/Session/History/SSessionHistory.h
@@ -39,6 +39,8 @@ public:
 	
 		SLATE_EVENT(SConcertSessionActivities::FGetPackageEvent, GetPackageEvent)
 		SLATE_EVENT(SConcertSessionActivities::FGetTransactionEvent, GetTransactionEvent)
+		/** If bound, invoked when an item in the table is right-click and we are supposed to show a menu of actions. */
+		SLATE_EVENT(FOnContextMenuOpening, OnContextMenuOpening)
 	
 		/** Optional name for only displaying activities that affect certain packages */
 		SLATE_ARGUMENT(FName, PackageFilter)
