Fix pressing delete on group items in session browser causing a crash

diff --git a/Engine/Plugins/Developer/Concert/ConcertApp/MultiUserServer/Source/MultiUserServer/Private/Widgets/Browser/SConcertServerSessionBrowser.cpp b/Engine/Plugins/Developer/Concert/ConcertApp/MultiUserServer/Source/MultiUserServer/Private/Widgets/Browser/SConcertServerSessionBrowser.cpp
--- a/Engine/Plugins/Developer/Concert/ConcertApp/MultiUserServer/Source/MultiUserServer/Private/Widgets/Browser/SConcertServerSessionBrowser.cpp
+++ b/Engine/Plugins/Developer/Concert/ConcertApp/MultiUserServer/Source/MultiUserServer/Private/Widgets/Browser/SConcertServerSessionBrowser.cpp
@@ -83,6 +83,11 @@ namespace UE::MultiUserServer
 
 	void SConcertServerSessionBrowser::RequestDeleteSession(const TArray<TSharedPtr<FConcertSessionTreeItem>>& SessionItems) 
 	{
+		if (SessionItems.Num() == 0)
+		{
+			return;
+		}
+		
 		const FText Message = [this, &SessionItems]()
 		{
 			if (SessionItems.Num() > 1)
diff --git a/Engine/Plugins/Developer/Concert/ConcertUI/ConcertSharedSlate/Source/ConcertSharedSlate/Private/Session/Browser/SConcertSessionBrowser.cpp b/Engine/Plugins/Developer/Concert/ConcertUI/ConcertSharedSlate/Source/ConcertSharedSlate/Private/Session/Browser/SConcertSessionBrowser.cpp
--- a/Engine/Plugins/Developer/Concert/ConcertUI/ConcertSharedSlate/Source/ConcertSharedSlate/Private/Session/Browser/SConcertSessionBrowser.cpp
+++ b/Engine/Plugins/Developer/Concert/ConcertUI/ConcertSharedSlate/Source/ConcertSharedSlate/Private/Session/Browser/SConcertSessionBrowser.cpp
@@ -1226,6 +1226,11 @@ void SConcertSessionBrowser::RequestRenameSession(const TSharedPtr<FConcertSessi
 
 void SConcertSessionBrowser::RequestDeleteSessions(const TArray<TSharedPtr<FConcertSessionTreeItem>>& ItemsToDelete)
 {
+	if (ItemsToDelete.Num() == 0)
+	{
+		return;
+	}
+	
 	if (!AskUserToDeleteSessions.IsBound() || AskUserToDeleteSessions.Execute(ItemsToDelete))
 	{
 		ConcertBrowserUtils::RequestItemDeletion(*GetController().Get(), ItemsToDelete);
