Fix color (pink -> white) for Multi-user server style

diff --git a/Engine/Plugins/Developer/Concert/ConcertUI/ConcertSharedSlate/Source/ConcertSharedSlate/Private/ConcertFrontendStyle.cpp b/Engine/Plugins/Developer/Concert/ConcertUI/ConcertSharedSlate/Source/ConcertSharedSlate/Private/ConcertFrontendStyle.cpp
--- a/Engine/Plugins/Developer/Concert/ConcertUI/ConcertSharedSlate/Source/ConcertSharedSlate/Private/ConcertFrontendStyle.cpp
+++ b/Engine/Plugins/Developer/Concert/ConcertUI/ConcertSharedSlate/Source/ConcertSharedSlate/Private/ConcertFrontendStyle.cpp
@@ -51,7 +51,7 @@ void FConcertFrontendStyle::Initialize()
 	// Use this to change the opacity. Ex: In UE4, the icon looked fine at 80%, in UE5, icons looks homogeneous with others at 100%.
 	const FLinearColor IconColorAndOpacity(FLinearColor(1.f, 1.f, 1.f, 1.f));
 	const FLinearColor ActiveIconColorAndOpacity(FLinearColor(1.f, 1.f, 1.f, 1.f));
-	const FLinearColor ArchivedIconColorAndOpacity(FLinearColor(0.9f, 0.13f, 0.54f, 1.f));
+	const FLinearColor ArchivedIconColorAndOpacity(FLinearColor(1.f, 1.f, 1.f, 1.f));
 
 	// Colors
 	StyleSet->Set("Concert.ActiveSession.Color", ActiveIconColorAndOpacity);
