Add object combo box now inlines actors in the level to reduce number of clicks required by one

diff --git a/Engine/Plugins/Developer/Concert/ConcertUI/ConcertClientSharedSlate/Source/ConcertClientSharedSlate/Private/Replication/Editor/Model/Object/EditorObjectSelectionSourceModel.cpp b/Engine/Plugins/Developer/Concert/ConcertUI/ConcertClientSharedSlate/Source/ConcertClientSharedSlate/Private/Replication/Editor/Model/Object/EditorObjectSelectionSourceModel.cpp
--- a/Engine/Plugins/Developer/Concert/ConcertUI/ConcertClientSharedSlate/Source/ConcertClientSharedSlate/Private/Replication/Editor/Model/Object/EditorObjectSelectionSourceModel.cpp
+++ b/Engine/Plugins/Developer/Concert/ConcertUI/ConcertClientSharedSlate/Source/ConcertClientSharedSlate/Private/Replication/Editor/Model/Object/EditorObjectSelectionSourceModel.cpp
@@ -22,8 +22,8 @@ namespace UE::ConcertClientSharedSlate
 					LOCTEXT("ObjectsCategory.Tooltip", "Options for adding objects from the open editor world")
 				},
 				{
-					MakeShared<FWorldActorSource>(),
-					MakeShared<FSelectedActorsSource>()
+					MakeShared<FSelectedActorsSource>(),
+					MakeShared<FWorldActorSource>()
 				}
 			};
 		}())
diff --git a/Engine/Plugins/Developer/Concert/ConcertUI/ConcertClientSharedSlate/Source/ConcertClientSharedSlate/Private/Replication/Editor/View/ObjectEditor/SBaseReplicationStreamEditor.cpp b/Engine/Plugins/Developer/Concert/ConcertUI/ConcertClientSharedSlate/Source/ConcertClientSharedSlate/Private/Replication/Editor/View/ObjectEditor/SBaseReplicationStreamEditor.cpp
--- a/Engine/Plugins/Developer/Concert/ConcertUI/ConcertClientSharedSlate/Source/ConcertClientSharedSlate/Private/Replication/Editor/View/ObjectEditor/SBaseReplicationStreamEditor.cpp
+++ b/Engine/Plugins/Developer/Concert/ConcertUI/ConcertClientSharedSlate/Source/ConcertClientSharedSlate/Private/Replication/Editor/View/ObjectEditor/SBaseReplicationStreamEditor.cpp
@@ -263,7 +263,8 @@ namespace UE::ConcertClientSharedSlate
 				return EditablePropertiesModel->ContainsObjects({ Item.Object.Get() } );
 			}),
 			TAttribute<bool>::CreateLambda([this]() { return !IsEditingDisabled(); }),
-			TAttribute<FText>::CreateLambda([this]() { return GetEditingDisabledText(); })
+			TAttribute<FText>::CreateLambda([this]() { return GetEditingDisabledText(); }),
+			EItemPickerFlags::DisplayOptionListInline
 		}; 
 	}
 }
diff --git a/Engine/Plugins/Developer/Concert/ConcertUI/ConcertSharedSlate/Source/ConcertSharedSlate/Public/Model/Item/SourceModelBuilders.h b/Engine/Plugins/Developer/Concert/ConcertUI/ConcertSharedSlate/Source/ConcertSharedSlate/Public/Model/Item/SourceModelBuilders.h
--- a/Engine/Plugins/Developer/Concert/ConcertUI/ConcertSharedSlate/Source/ConcertSharedSlate/Public/Model/Item/SourceModelBuilders.h
+++ b/Engine/Plugins/Developer/Concert/ConcertUI/ConcertSharedSlate/Source/ConcertSharedSlate/Public/Model/Item/SourceModelBuilders.h
@@ -17,8 +17,21 @@ class FBaseMenuBuilder;
 
 #define LOCTEXT_NAMESPACE "FSourceModelBuilders"
 
+/**
+ * These functions build widgets for selecting options.
+ * TSourceSelectionCategory is an entry in a combo button menu; it can contain subcategories.
+ * Categories contain IItemSourceModels which are a source of options.
+ */
 namespace UE::ConcertSharedSlate
 {
+	enum class EItemPickerFlags
+	{
+		None,
+		/** If the source type is ESourceType::ShowAsList, instead of creating a menu entry that spawns a submenu (default behavior) just place the source's items in the same menu. */
+		DisplayOptionListInline = 1 << 0
+	};
+	ENUM_CLASS_FLAGS(EItemPickerFlags);
+	
 	/**
 	 * Utility for creating widgets of IItemSourceModel and categories thereof (TSourceSelectionCategory).
 	 */
@@ -49,13 +62,17 @@ namespace UE::ConcertSharedSlate
 			/** Optional. If IsEnabled returns true, this tooltip is displayed on the relevant UI. */
 			TAttribute<FText> DisabledToolTipAttribute;
 
+			/** Special flags for altering default behavior. */
+			EItemPickerFlags Flags;
+
 			FItemPickerArgs(
 				FOnItemsSelected OnObjectsSelected,
 				FGetItemDisplayString GetItemDisplayString,
 				FGetItemIcon GetItemIcon = {},
 				FIsItemSelected IsItemSelected = {},
 				TAttribute<bool> IsEnabledAttribute = {},
-				TAttribute<FText> DisabledToolTipAttribute = {}
+				TAttribute<FText> DisabledToolTipAttribute = {},
+				EItemPickerFlags Flags = EItemPickerFlags::None
 				)
 				: OnItemsSelected(MoveTemp(OnObjectsSelected))
 				, GetItemDisplayString(MoveTemp(GetItemDisplayString))
@@ -63,6 +80,7 @@ namespace UE::ConcertSharedSlate
 				, IsItemSelected(MoveTemp(IsItemSelected))
 				, IsEnabledAttribute(MoveTemp(IsEnabledAttribute))
 				, DisabledToolTipAttribute(MoveTemp(DisabledToolTipAttribute))
+				, Flags(Flags)
 			{
 				check(this->OnItemsSelected.IsBound() && this->GetItemDisplayString.IsBound());
 			}
@@ -85,7 +103,10 @@ namespace UE::ConcertSharedSlate
 	private:
 		
 		static void FillSearchableSubmenu(const TSharedRef<IItemSourceModel<TItemType>>& Source, const FItemPickerArgs& Args, FBaseMenuBuilder& MenuBuilder);
+		
 		static TSharedRef<SWidget> BuildMenu(const TArray<TAttribute<TSharedPtr<IItemSourceModel<TItemType>>>>& Options, TArray<TSourceSelectionCategory<TItemType>> SubCategories, const FItemPickerArgs& Args);
+		static bool NeedsSeparationFromOtherItems(const TSharedRef<IItemSourceModel<TItemType>>& Option, const FItemPickerArgs& Args);
+		
 		static void AddSubCategoryToMenu(const TSourceSelectionCategory<TItemType>& Category, const FItemPickerArgs& Args, FBaseMenuBuilder& MenuBuilder);
 	};
 }
@@ -187,22 +208,29 @@ namespace UE::ConcertSharedSlate
 		{
 		case ESourceType::ShowAsList:
 			{
-				FMenuEntryParams ShowAsListParams;
-				ShowAsListParams.LabelOverride = DisplayInfo.Label;
-				ShowAsListParams.ToolTipOverride = DisplayInfo.ToolTip;
-				ShowAsListParams.bIsSubMenu = true;
-				// Dummy is needed to avoid assert
-				ShowAsListParams.DirectActions = {
-					FExecuteAction::CreateLambda([](){})
-					};
-				ShowAsListParams.MenuBuilder.BindLambda([Source = MoveTemp(Option), Args = MoveTemp(Args)]()
+				if (EnumHasAnyFlags(Args.Flags, EItemPickerFlags::DisplayOptionListInline))
 				{
-					FMenuBuilder MenuBuilder(true, nullptr);
-					FillSearchableSubmenu(Source, Args, MenuBuilder);
-					return MenuBuilder.MakeWidget();
-				});
+					FillSearchableSubmenu(Option, Args, MenuBuilder);
+				}
+				else
+				{
+					FMenuEntryParams ShowAsListParams;
+					ShowAsListParams.LabelOverride = DisplayInfo.Label;
+					ShowAsListParams.ToolTipOverride = DisplayInfo.ToolTip;
+					ShowAsListParams.bIsSubMenu = true;
+					// Dummy is needed to avoid assert
+					ShowAsListParams.DirectActions = {
+						FExecuteAction::CreateLambda([](){})
+						};
+					ShowAsListParams.MenuBuilder.BindLambda([Source = MoveTemp(Option), Args = MoveTemp(Args)]()
+					{
+						FMenuBuilder MenuBuilder(true, nullptr);
+						FillSearchableSubmenu(Source, Args, MenuBuilder);
+						return MenuBuilder.MakeWidget();
+					});
 				
-				MenuBuilder.AddMenuEntry(ShowAsListParams);
+					MenuBuilder.AddMenuEntry(ShowAsListParams);
+				}
 				break;
 			}
 		
@@ -266,12 +294,30 @@ namespace UE::ConcertSharedSlate
 		FMenuBuilder MenuBuilder(true, nullptr);
 
 		// Sources first...
-		for (const TAttribute<TSharedPtr<IItemSourceModel<TItemType>>>& Source : Options)
+		bool bPreviousSourceNeededSeparator = false;
+		for (int32 i = 0; i < Options.Num(); ++i)
 		{
+			const TAttribute<TSharedPtr<IItemSourceModel<TItemType>>>& Source = Options[i];
 			const TSharedPtr<IItemSourceModel<TItemType>> Model = Source.Get();
 			if (ensure(Model))
 			{
-				AddOptionToMenu(Model.ToSharedRef(), Args, MenuBuilder);
+				const TSharedRef<IItemSourceModel<TItemType>> ModelRef = Model.ToSharedRef();
+
+				// E.g. if the source is a list of (many) inlined items, they are separated
+				bPreviousSourceNeededSeparator = NeedsSeparationFromOtherItems(ModelRef, Args);
+				const bool bIsNotFirstItem = i > 0;
+				if (bPreviousSourceNeededSeparator && bIsNotFirstItem)
+				{
+					MenuBuilder.AddSeparator();
+				}
+				
+				AddOptionToMenu(ModelRef, Args, MenuBuilder);
+				
+				const bool bIsNotLastItem = i < Options.Num() - 1;
+				if (bPreviousSourceNeededSeparator && bIsNotLastItem)
+				{
+					MenuBuilder.AddSeparator();
+				}
 			}
 		}
 			
@@ -281,6 +327,10 @@ namespace UE::ConcertSharedSlate
 			// Is this slow? May want to cache ToString(). The number of items is so low it will hopefully have no noticeable effect.
 			return Left.DisplayInfo.Label.ToString() < Right.DisplayInfo.Label.ToString();
 		});
+		if (!SubCategories.IsEmpty() && bPreviousSourceNeededSeparator)
+		{
+			MenuBuilder.AddSeparator();
+		}
 		for (const TSourceSelectionCategory<TItemType>& SubCategory : SubCategories)
 		{
 			AddSubCategoryToMenu(SubCategory, Args, MenuBuilder);
@@ -289,6 +339,14 @@ namespace UE::ConcertSharedSlate
 		return MenuBuilder.MakeWidget();
 	}
 
+	template <typename TItemType>
+	bool FSourceModelBuilders<TItemType>::NeedsSeparationFromOtherItems(const TSharedRef<IItemSourceModel<TItemType>>& Option, const FItemPickerArgs& Args)
+	{
+		const bool bIsInlinedItemList = Option->GetDisplayInfo().SourceType == ESourceType::ShowAsList
+			&& EnumHasAnyFlags(Args.Flags, EItemPickerFlags::DisplayOptionListInline);
+		return bIsInlinedItemList;
+	}
+
 	template<typename TItemType>
 	void FSourceModelBuilders<TItemType>::AddSubCategoryToMenu(const TSourceSelectionCategory<TItemType>& Category, const FItemPickerArgs& Args, FBaseMenuBuilder& MenuBuilder)
 	{
