Extract IObjectNameModel from ISubobjectModel

diff --git a/Engine/Plugins/Developer/Concert/ConcertApp/MultiUserClient/Source/MultiUserClient/Private/Widgets/ActiveSession/Replication/Client/Multi/SMultiClientView.cpp b/Engine/Plugins/Developer/Concert/ConcertApp/MultiUserClient/Source/MultiUserClient/Private/Widgets/ActiveSession/Replication/Client/Multi/SMultiClientView.cpp
--- a/Engine/Plugins/Developer/Concert/ConcertApp/MultiUserClient/Source/MultiUserClient/Private/Widgets/ActiveSession/Replication/Client/Multi/SMultiClientView.cpp
+++ b/Engine/Plugins/Developer/Concert/ConcertApp/MultiUserClient/Source/MultiUserClient/Private/Widgets/ActiveSession/Replication/Client/Multi/SMultiClientView.cpp
@@ -93,6 +93,7 @@ namespace UE::MultiUserClient
 			.ViewerParams 
 			{
 				.SubobjectModel = ConcertClientSharedSlate::CreateSubobjectModelForComponentHierarchy(), // This makes actors have children in the top view
+				.NameModel = ConcertClientSharedSlate::CreateEditorObjectNameModel(), // This makes actors use their labels, and components use the names given in the BP editor
 				.OnExtendObjectsContextMenu = FExtendObjectMenu::CreateSP(this, &SMultiClientView::ExtendObjectContextMenu),
 				.AdditionalObjectColumns =
 				{
diff --git a/Engine/Plugins/Developer/Concert/ConcertApp/MultiUserClient/Source/MultiUserClient/Private/Widgets/ActiveSession/Replication/Client/Single/SReplicationClientView.cpp b/Engine/Plugins/Developer/Concert/ConcertApp/MultiUserClient/Source/MultiUserClient/Private/Widgets/ActiveSession/Replication/Client/Single/SReplicationClientView.cpp
--- a/Engine/Plugins/Developer/Concert/ConcertApp/MultiUserClient/Source/MultiUserClient/Private/Widgets/ActiveSession/Replication/Client/Single/SReplicationClientView.cpp
+++ b/Engine/Plugins/Developer/Concert/ConcertApp/MultiUserClient/Source/MultiUserClient/Private/Widgets/ActiveSession/Replication/Client/Single/SReplicationClientView.cpp
@@ -79,6 +79,7 @@ namespace UE::MultiUserClient
 			.ViewerParams =
 			{
 				.SubobjectModel = ConcertClientSharedSlate::CreateSubobjectModelForComponentHierarchy(), // This makes actors have children in the top view
+				.NameModel = ConcertClientSharedSlate::CreateEditorObjectNameModel(), // This makes actors use their labels, and components use the names given in the BP editor
 				.OnExtendObjectsContextMenu = FExtendObjectMenu::CreateSP(this, &SReplicationClientView::ExtendObjectContextMenu),
 				.AdditionalObjectColumns =
 				{
diff --git a/Engine/Plugins/Developer/Concert/ConcertApp/MultiUserClient/Source/MultiUserReplicationEditor/Private/Assets/AssetDefinition_ReplicationSessionPreset.h b/Engine/Plugins/Developer/Concert/ConcertApp/MultiUserClient/Source/MultiUserReplicationEditor/Private/Assets/AssetDefinition_ReplicationSessionPreset.h
--- a/Engine/Plugins/Developer/Concert/ConcertApp/MultiUserClient/Source/MultiUserReplicationEditor/Private/Assets/AssetDefinition_ReplicationSessionPreset.h
+++ b/Engine/Plugins/Developer/Concert/ConcertApp/MultiUserClient/Source/MultiUserReplicationEditor/Private/Assets/AssetDefinition_ReplicationSessionPreset.h
@@ -9,7 +9,7 @@
  * 
  */
 UCLASS()
-class MULTIUSERREPLICATIONEDITOR_API UAssetDefinition_ReplicationSessionPreset : public UAssetDefinitionDefault
+class UAssetDefinition_ReplicationSessionPreset : public UAssetDefinitionDefault
 {
 	GENERATED_BODY()
 public:
@@ -21,4 +21,17 @@ public:
 	virtual TConstArrayView<FAssetCategoryPath> GetAssetCategories() const override;
 	EAssetCommandResult OpenAssets(const FAssetOpenArgs& OpenArgs) const override;
 	//~ End UAssetDefinition Interface
+
+protected:
+	
+	//~ Begin UAssetDefinition Interface
+	virtual bool CanRegisterStatically() const override
+	{
+		// TODO UE-196506:
+		// The entire replication feature is a MVP created in 5.4.
+		// Offline editing is not exposed to protect end-users from relying on it until we figure out workflows & requirements.
+		// When implementing offline stream editing adjust the editor this asset definition creates, and remove this.
+		return false;
+	}
+	//~ End UAssetDefinition Interface
 };
diff --git a/Engine/Plugins/Developer/Concert/ConcertApp/MultiUserClient/Source/MultiUserReplicationEditor/Private/Assets/ReplicationSessionPresetFactory.cpp b/Engine/Plugins/Developer/Concert/ConcertApp/MultiUserClient/Source/MultiUserReplicationEditor/Private/Assets/ReplicationSessionPresetFactory.cpp
--- a/Engine/Plugins/Developer/Concert/ConcertApp/MultiUserClient/Source/MultiUserReplicationEditor/Private/Assets/ReplicationSessionPresetFactory.cpp
+++ b/Engine/Plugins/Developer/Concert/ConcertApp/MultiUserClient/Source/MultiUserReplicationEditor/Private/Assets/ReplicationSessionPresetFactory.cpp
@@ -16,7 +16,10 @@ UReplicationSessionPresetFactory::UReplicationSessionPresetFactory(const FObject
 
 bool UReplicationSessionPresetFactory::CanCreateNew() const
 {
-	return true;
+	// TODO UE-196506:
+	// The entire replication feature is a MVP created in 5.4.
+	// Offline editing is not exposed to protect end-users from relying on it until we figure out workflows & requirements.
+	return false;
 }
 
 FText UReplicationSessionPresetFactory::GetDisplayName() const
@@ -32,6 +35,8 @@ bool UReplicationSessionPresetFactory::ConfigureProperties()
 UObject* UReplicationSessionPresetFactory::FactoryCreateNew(UClass* Class, UObject* InParent, FName Name, EObjectFlags Flags, UObject* Context, FFeedbackContext* Warn)
 {
 	UMultiUserReplicationSessionPreset* Asset = NewObject<UMultiUserReplicationSessionPreset>(InParent, Name, Flags);
+	// TODO UE-196506:
+	ensureMsgf(false, TEXT("Should not be called in MVP"));
 	return Asset;
 }
 
diff --git a/Engine/Plugins/Developer/Concert/ConcertUI/ConcertClientSharedSlate/Source/ConcertClientSharedSlate/Private/Replication/ClientReplicationWidgetFactories.cpp b/Engine/Plugins/Developer/Concert/ConcertUI/ConcertClientSharedSlate/Source/ConcertClientSharedSlate/Private/Replication/ClientReplicationWidgetFactories.cpp
--- a/Engine/Plugins/Developer/Concert/ConcertUI/ConcertClientSharedSlate/Source/ConcertClientSharedSlate/Private/Replication/ClientReplicationWidgetFactories.cpp
+++ b/Engine/Plugins/Developer/Concert/ConcertUI/ConcertClientSharedSlate/Source/ConcertClientSharedSlate/Private/Replication/ClientReplicationWidgetFactories.cpp
@@ -2,6 +2,7 @@
 
 #include "Replication/ClientReplicationWidgetFactories.h"
 
+#include "Object/EditorObjectNameModel.h"
 #include "Replication/ReplicationWidgetFactories.h"
 #include "Replication/Editor/Model/ReplicationStreamObject.h"
 #include "Replication/Editor/Model/Subobject/ComponentHierarchySubobjectModel.h"
@@ -18,6 +19,11 @@ namespace UE::ConcertClientSharedSlate
 		return MakeShared<FComponentHierarchySubobjectModel>();
 	}
 
+	TSharedRef<ConcertSharedSlate::IObjectNameModel> CreateEditorObjectNameModel()
+	{
+		return MakeShared<FEditorObjectNameModel>();
+	}
+
 	TSharedRef<ConcertSharedSlate::IEditableReplicationStreamModel> CreateTransactionalStreamModel(
 		TSharedRef<ConcertSharedSlate::IEditableReplicationStreamModel> BaseModel,
 		UObject& OwnerObject
diff --git a/Engine/Plugins/Developer/Concert/ConcertUI/ConcertClientSharedSlate/Source/ConcertClientSharedSlate/Private/Replication/Editor/Model/Subobject/ComponentHierarchySubobjectModel.cpp b/Engine/Plugins/Developer/Concert/ConcertUI/ConcertClientSharedSlate/Source/ConcertClientSharedSlate/Private/Replication/Editor/Model/Subobject/ComponentHierarchySubobjectModel.cpp
--- a/Engine/Plugins/Developer/Concert/ConcertUI/ConcertClientSharedSlate/Source/ConcertClientSharedSlate/Private/Replication/Editor/Model/Subobject/ComponentHierarchySubobjectModel.cpp
+++ b/Engine/Plugins/Developer/Concert/ConcertUI/ConcertClientSharedSlate/Source/ConcertClientSharedSlate/Private/Replication/Editor/Model/Subobject/ComponentHierarchySubobjectModel.cpp
@@ -45,7 +45,6 @@ namespace UE::ConcertClientSharedSlate
 					}
 				}
 			}
-			OnHierarchyChangedDelegate.Broadcast();
 		}
 	}
 
@@ -54,12 +53,6 @@ namespace UE::ConcertClientSharedSlate
 		return ConcertSharedSlate::ObjectUtils::IsActor(Object);
 	}
 
-	FText FComponentHierarchySubobjectModel::GetSubobjectDisplayName(const FSoftObjectPath& ObjectPath) const
-	{
-		const FObjectMetaData* MetaData = ObjectMetaData.Find(ObjectPath);
-		return MetaData ? MetaData->DisplayLabel : FText::GetEmpty();
-	}
-
 	void FComponentHierarchySubobjectModel::ForEachRootSubobject(FName Category, TFunctionRef<EBreakBehavior(const FSoftObjectPath& Object)> Callback) const
 	{
 		const AActor* Actor = GetTopLevelObjectAsActor();
diff --git a/Engine/Plugins/Developer/Concert/ConcertUI/ConcertClientSharedSlate/Source/ConcertClientSharedSlate/Private/Replication/Editor/Model/Subobject/ComponentHierarchySubobjectModel.h b/Engine/Plugins/Developer/Concert/ConcertUI/ConcertClientSharedSlate/Source/ConcertClientSharedSlate/Private/Replication/Editor/Model/Subobject/ComponentHierarchySubobjectModel.h
--- a/Engine/Plugins/Developer/Concert/ConcertUI/ConcertClientSharedSlate/Source/ConcertClientSharedSlate/Private/Replication/Editor/Model/Subobject/ComponentHierarchySubobjectModel.h
+++ b/Engine/Plugins/Developer/Concert/ConcertUI/ConcertClientSharedSlate/Source/ConcertClientSharedSlate/Private/Replication/Editor/Model/Subobject/ComponentHierarchySubobjectModel.h
@@ -27,10 +27,8 @@ namespace UE::ConcertClientSharedSlate
 		virtual FSoftObjectPath GetTopLevelObject() const override { return TopLevelObject; }
 		virtual bool IsTopLevelObject(const FSoftObjectPath& Object) const override;
 		virtual TArray<FName> GetCategories() const override { return { SceneComponentsCategory, ActorComponentsCategory }; }
-		virtual FText GetSubobjectDisplayName(const FSoftObjectPath& ObjectPath) const override;
 		virtual void ForEachRootSubobject(FName Category, TFunctionRef<EBreakBehavior(const FSoftObjectPath& Object)> Callback) const override;
 		virtual void ForEachDirectChildSubobject(const FSoftObjectPath& Parent, TFunctionRef<EBreakBehavior(const FSoftObjectPath&)> Callback) const override;
-		virtual FOnHierarchyChanged& OnHierarchyChanged() override { return OnHierarchyChangedDelegate; }
 		//~ End ISubobjectModel Interface
 
 	private:
@@ -44,9 +42,6 @@ namespace UE::ConcertClientSharedSlate
 		};
 		TMap<FSoftObjectPath,FObjectMetaData> ObjectMetaData;
 
-		/** Called when the hierarchy has changed, e.g. due to calling SetTopLevelObject. */
-		FOnHierarchyChanged OnHierarchyChangedDelegate;
-
 		AActor* GetTopLevelObjectAsActor() const;
 	};
 }
diff --git a/Engine/Plugins/Developer/Concert/ConcertUI/ConcertClientSharedSlate/Source/ConcertClientSharedSlate/Private/Replication/Object/EditorObjectNameModel.cpp b/Engine/Plugins/Developer/Concert/ConcertUI/ConcertClientSharedSlate/Source/ConcertClientSharedSlate/Private/Replication/Object/EditorObjectNameModel.cpp
new file mode 100644
--- /dev/null
+++ b/Engine/Plugins/Developer/Concert/ConcertUI/ConcertClientSharedSlate/Source/ConcertClientSharedSlate/Private/Replication/Object/EditorObjectNameModel.cpp
@@ -0,0 +1,56 @@
+// Copyright Epic Games, Inc. All Rights Reserved.
+
+#include "EditorObjectNameModel.h"
+
+#include "Internationalization/Text.h"
+#include "Replication/ObjectUtils.h"
+
+#include "GameFramework/Actor.h"
+#include "SubobjectDataSubsystem.h"
+
+namespace UE::ConcertClientSharedSlate
+{
+	namespace Private
+	{
+		static FText FindSubobjectDisplayName(const UObject& Subbject, AActor& OwningActor)
+		{
+			USubobjectDataSubsystem* SubobjectDataSubsystem = USubobjectDataSubsystem::Get();
+			TArray<FSubobjectDataHandle> Handles;
+			SubobjectDataSubsystem->GatherSubobjectData(&OwningActor, Handles);
+				
+			for (const FSubobjectDataHandle& Handle : Handles)
+			{
+				const FSubobjectData* SubobjectData = Handle.GetData();
+				const UObject* Object = SubobjectData->FindComponentInstanceInActor(&OwningActor);
+				if (Object == &Subbject)
+				{
+					return FText::FromString(SubobjectData->GetDisplayString());
+				}
+			}
+
+			return FText::GetEmpty();
+		}
+	}
+	
+	FText FEditorObjectNameModel::GetObjectDisplayName(const FSoftObjectPath& ObjectPath) const
+	{
+		if (UObject* ResolvedObject = ObjectPath.ResolveObject())
+		{
+			// Display actor just like the outliner does
+			if (const AActor* AsActor = Cast<AActor>(ResolvedObject))
+			{
+				return FText::FromString(AsActor->GetActorLabel());
+			}
+
+			// Display the same component name as the SSubobjectEditor widget does, i.e. component hierarchy in the details panel or Blueprint editor.
+			AActor* OwningActor = ResolvedObject->GetTypedOuter<AActor>();
+			const FText FoundSubobjectName = OwningActor ? Private::FindSubobjectDisplayName(*ResolvedObject, *OwningActor) : FText::GetEmpty();
+			if (!FoundSubobjectName.IsEmpty())
+			{
+				return FoundSubobjectName;
+			}
+		}
+		
+		return FText::FromString(ConcertSharedSlate::ObjectUtils::ExtractObjectDisplayStringFromPath(ObjectPath));
+	}
+}
diff --git a/Engine/Plugins/Developer/Concert/ConcertUI/ConcertClientSharedSlate/Source/ConcertClientSharedSlate/Private/Replication/Object/EditorObjectNameModel.h b/Engine/Plugins/Developer/Concert/ConcertUI/ConcertClientSharedSlate/Source/ConcertClientSharedSlate/Private/Replication/Object/EditorObjectNameModel.h
new file mode 100644
--- /dev/null
+++ b/Engine/Plugins/Developer/Concert/ConcertUI/ConcertClientSharedSlate/Source/ConcertClientSharedSlate/Private/Replication/Object/EditorObjectNameModel.h
@@ -0,0 +1,19 @@
+// Copyright Epic Games, Inc. All Rights Reserved.
+
+#pragma once
+
+#include "Replication/Editor/Model/Object/IObjectNameModel.h"
+
+namespace UE::ConcertClientSharedSlate
+{
+	/** Name model that uses editor data for determining display names: actors use their labels, components ask USubobjectDataSubsystem.  */
+	class FEditorObjectNameModel : public ConcertSharedSlate::IObjectNameModel
+	{
+	public:
+
+		//~ Begin IObjectNameModel Interface
+		virtual FText GetObjectDisplayName(const FSoftObjectPath& ObjectPath) const override;
+		//~ End IObjectNameModel Interface
+	};
+}
+
diff --git a/Engine/Plugins/Developer/Concert/ConcertUI/ConcertClientSharedSlate/Source/ConcertClientSharedSlate/Public/Replication/ClientReplicationWidgetFactories.h b/Engine/Plugins/Developer/Concert/ConcertUI/ConcertClientSharedSlate/Source/ConcertClientSharedSlate/Public/Replication/ClientReplicationWidgetFactories.h
--- a/Engine/Plugins/Developer/Concert/ConcertUI/ConcertClientSharedSlate/Source/ConcertClientSharedSlate/Public/Replication/ClientReplicationWidgetFactories.h
+++ b/Engine/Plugins/Developer/Concert/ConcertUI/ConcertClientSharedSlate/Source/ConcertClientSharedSlate/Public/Replication/ClientReplicationWidgetFactories.h
@@ -10,6 +10,7 @@ struct FConcertObjectReplicationMap;
 namespace UE::ConcertSharedSlate
 {
 	class IEditableReplicationStreamModel;
+	class IObjectNameModel;
 	class IReplicationStreamEditor;
 	class IStreamExtender;
 	class ISubobjectModel;
@@ -20,6 +21,9 @@ namespace UE::ConcertClientSharedSlate
 {
 	/** Builds a similar tree hierarchy as SSubobjectEditor. Reports only components as subobjects. */
 	CONCERTCLIENTSHAREDSLATE_API TSharedRef<ConcertSharedSlate::ISubobjectModel> CreateSubobjectModelForComponentHierarchy();
+
+	/** Name model that uses editor data for determining display names: actors use their labels, components ask USubobjectDataSubsystem. */
+	CONCERTCLIENTSHAREDSLATE_API TSharedRef<ConcertSharedSlate::IObjectNameModel> CreateEditorObjectNameModel();
 	
 	/**
 	 * Wraps the passed in BaseModel and makes it transactional.
diff --git a/Engine/Plugins/Developer/Concert/ConcertUI/ConcertSharedSlate/Source/ConcertSharedSlate/Private/Replication/Editor/View/DisplayUtils.cpp b/Engine/Plugins/Developer/Concert/ConcertUI/ConcertSharedSlate/Source/ConcertSharedSlate/Private/Replication/Editor/View/DisplayUtils.cpp
--- a/Engine/Plugins/Developer/Concert/ConcertUI/ConcertSharedSlate/Source/ConcertSharedSlate/Private/Replication/Editor/View/DisplayUtils.cpp
+++ b/Engine/Plugins/Developer/Concert/ConcertUI/ConcertSharedSlate/Source/ConcertSharedSlate/Private/Replication/Editor/View/DisplayUtils.cpp
@@ -3,7 +3,9 @@
 #include "Replication/Editor/View/DisplayUtils.h"
 
 #include "ConcertLogGlobal.h"
+#include "Replication/ObjectUtils.h"
 #include "Replication/Editor/Model/IReplicationStreamModel.h"
+#include "Replication/Editor/Model/Object/IObjectNameModel.h"
 
 #include "Styling/SlateIconFinder.h"
 
@@ -16,29 +18,27 @@
 
 namespace UE::ConcertSharedSlate::DisplayUtils
 {
-	FText GetObjectDisplayText(const FSoftObjectPath& Object)
+	FText GetObjectDisplayText(const FSoftObjectPath& Object, IObjectNameModel* Model)
 	{
-		return FText::FromString(GetObjectDisplayString(Object));
-	}
-
-	FString GetObjectDisplayString(const FSoftObjectPath& Object)
-	{
-		// Important! The object may not be loaded, yet. This could be if the asset is using a level that was not opened.
+		if (Model)
+		{
+			return Model->GetObjectDisplayName(Object);
+		}
+		
 #if WITH_EDITOR
+		// Important! The object may not be loaded, yet. So only resolve, do not TryLoad. This could be if the asset is using a level that was not opened.
 		if (const UObject* LoadedObject = Object.ResolveObject())
 		{
-			return GetObjectDisplayString(*LoadedObject);
+			return FText::FromString(GetObjectDisplayString(*LoadedObject));
 		}
 #endif
 
-		// Subpath looks like this PersistentLevel.Actor.Component
-		const FString& Subpath = Object.GetSubPathString();
-		const int32 LastDotIndex = Subpath.Find(TEXT("."), ESearchCase::CaseSensitive, ESearchDir::FromEnd);
-		if (LastDotIndex == INDEX_NONE)
-		{
-			return {};
-		}
-		return Subpath.RightChop(LastDotIndex + 1);
+		return ExtractObjectDisplayTextFromPath(Object);
+	}
+	
+	FText ExtractObjectDisplayTextFromPath(const FSoftObjectPath& Object)
+	{
+		return FText::FromString(ObjectUtils::ExtractObjectDisplayStringFromPath(Object));
 	}
 
 	FString GetObjectDisplayString(const UObject& Object)
diff --git a/Engine/Plugins/Developer/Concert/ConcertUI/ConcertSharedSlate/Source/ConcertSharedSlate/Private/Replication/Editor/View/DisplayUtils.h b/Engine/Plugins/Developer/Concert/ConcertUI/ConcertSharedSlate/Source/ConcertSharedSlate/Private/Replication/Editor/View/DisplayUtils.h
--- a/Engine/Plugins/Developer/Concert/ConcertUI/ConcertSharedSlate/Source/ConcertSharedSlate/Private/Replication/Editor/View/DisplayUtils.h
+++ b/Engine/Plugins/Developer/Concert/ConcertUI/ConcertSharedSlate/Source/ConcertSharedSlate/Private/Replication/Editor/View/DisplayUtils.h
@@ -2,8 +2,10 @@
 
 #pragma once
 
+#include "Templates/SharedPointer.h"
+#include "Internationalization/Text.h"
+
 class FString;
-class FText;
 class UObject;
 
 struct FConcertPropertyChain;
@@ -12,16 +14,19 @@ struct FSoftObjectPath;
 
 namespace UE::ConcertSharedSlate
 {
+	class IObjectNameModel;
 	class IReplicationStreamModel;
 }
 
 namespace UE::ConcertSharedSlate::DisplayUtils
 {
+	/** Unified version that reads from IObjectNameModel or defaults to extracting the name from the path. */
+	FText GetObjectDisplayText(const FSoftObjectPath& Object, IObjectNameModel* Model = nullptr);
+	inline FText GetObjectDisplayText(const FSoftObjectPath& Object, const TSharedPtr<IObjectNameModel>& Model = nullptr) { return GetObjectDisplayText(Object, Model.Get()); }
+	
 	/** @return The text to use for displaying this object's name */
-	FText GetObjectDisplayText(const FSoftObjectPath& Object);
+	FText ExtractObjectDisplayTextFromPath(const FSoftObjectPath& Object);
 	
-	/** @return More lightweight version of GetObjectDisplayText which does not construct any FText. */
-	FString GetObjectDisplayString(const FSoftObjectPath& Object);
 	/** @return More lightweight version of GetObjectDisplayText which does not construct any FText. */
 	FString GetObjectDisplayString(const UObject& Object);
 	
diff --git a/Engine/Plugins/Developer/Concert/ConcertUI/ConcertSharedSlate/Source/ConcertSharedSlate/Private/Replication/Editor/View/ObjectEditor/SBaseReplicationStreamEditor.cpp b/Engine/Plugins/Developer/Concert/ConcertUI/ConcertSharedSlate/Source/ConcertSharedSlate/Private/Replication/Editor/View/ObjectEditor/SBaseReplicationStreamEditor.cpp
--- a/Engine/Plugins/Developer/Concert/ConcertUI/ConcertSharedSlate/Source/ConcertSharedSlate/Private/Replication/Editor/View/ObjectEditor/SBaseReplicationStreamEditor.cpp
+++ b/Engine/Plugins/Developer/Concert/ConcertUI/ConcertSharedSlate/Source/ConcertSharedSlate/Private/Replication/Editor/View/ObjectEditor/SBaseReplicationStreamEditor.cpp
@@ -50,6 +50,7 @@ namespace UE::ConcertSharedSlate
 				.PrimaryPropertySort(InArgs._PrimaryPropertySort)
 				.SecondaryPropertySort(InArgs._SecondaryPropertySort)
 				.SubobjectModel(InArgs._SubobjectModel)
+				.NameModel(InArgs._NameModel)
 				.OnDeleteObjects(this, &SBaseReplicationStreamEditor::OnDeleteObjects)
 				.OnObjectsContextMenuOpening(this, &SBaseReplicationStreamEditor::OnObjectsContextMenuOpening)
 				.LeftOfObjectSearchBar()
diff --git a/Engine/Plugins/Developer/Concert/ConcertUI/ConcertSharedSlate/Source/ConcertSharedSlate/Private/Replication/Editor/View/ObjectEditor/SBaseReplicationStreamEditor.h b/Engine/Plugins/Developer/Concert/ConcertUI/ConcertSharedSlate/Source/ConcertSharedSlate/Private/Replication/Editor/View/ObjectEditor/SBaseReplicationStreamEditor.h
--- a/Engine/Plugins/Developer/Concert/ConcertUI/ConcertSharedSlate/Source/ConcertSharedSlate/Private/Replication/Editor/View/ObjectEditor/SBaseReplicationStreamEditor.h
+++ b/Engine/Plugins/Developer/Concert/ConcertUI/ConcertSharedSlate/Source/ConcertSharedSlate/Private/Replication/Editor/View/ObjectEditor/SBaseReplicationStreamEditor.h
@@ -20,10 +20,10 @@ namespace UE::ConcertSharedSlate
 	class FFakeObjectToPropertiesEditorModel;
 	class FReplicatedObjectData;
 	class FReplicatedPropertyData;
-	class IReplicationSubobjectView;
+	class IObjectNameModel;
 	class IObjectSelectionSourceModel;
-	class IReplicationStreamModel;
 	class IPropertySelectionSourceModel;
+	class IReplicationStreamModel;
 	class IReplicationSubobjectView;
 	class ISubobjectModel;
 	class SReplicationStreamViewer;
@@ -68,6 +68,8 @@ namespace UE::ConcertSharedSlate
 		
 			/** Optional. If set, this determines the children nested under the root objects. */
 			SLATE_ARGUMENT(TSharedPtr<ISubobjectModel>, SubobjectModel)
+			/** Optional. If set, this determines the display text for objects. */
+			SLATE_ARGUMENT(TSharedPtr<IObjectNameModel>, NameModel)
 
 			/** Called to generate the context menu for objects. This extends the options already generated by this widget. */
 			SLATE_EVENT(FExtendObjectMenu, OnExtendObjectsContextMenu)
diff --git a/Engine/Plugins/Developer/Concert/ConcertUI/ConcertSharedSlate/Source/ConcertSharedSlate/Private/Replication/Editor/View/ObjectViewer/Property/SPropertyTreeView.cpp b/Engine/Plugins/Developer/Concert/ConcertUI/ConcertSharedSlate/Source/ConcertSharedSlate/Private/Replication/Editor/View/ObjectViewer/Property/SPropertyTreeView.cpp
--- a/Engine/Plugins/Developer/Concert/ConcertUI/ConcertSharedSlate/Source/ConcertSharedSlate/Private/Replication/Editor/View/ObjectViewer/Property/SPropertyTreeView.cpp
+++ b/Engine/Plugins/Developer/Concert/ConcertUI/ConcertSharedSlate/Source/ConcertSharedSlate/Private/Replication/Editor/View/ObjectViewer/Property/SPropertyTreeView.cpp
@@ -55,6 +55,7 @@ namespace UE::ConcertSharedSlate
 	void SPropertyTreeView::Construct(const FArguments& InArgs)
 	{
 		SelectedObjectsAttribute = InArgs._SelectedObjects;
+		NameModel = InArgs._NameModel;
 		
 		const FBuildFilterBarResult Filters = BuildFilterBar();
 		ChildSlot
@@ -188,9 +189,9 @@ namespace UE::ConcertSharedSlate
 		}
 
 		TSet<FString> Names;
-		Algo::Transform(Objects, Names, [](const FSoftObjectPath& ObjectPath)
+		Algo::Transform(Objects, Names, [this](const FSoftObjectPath& ObjectPath)
 		{
-			return DisplayUtils::GetObjectDisplayText(ObjectPath).ToString();
+			return DisplayUtils::GetObjectDisplayText(ObjectPath, NameModel).ToString();
 		});
 		
 		const FText ObjectsText = FText::FromString(FString::Join(Names, TEXT(", ")));
diff --git a/Engine/Plugins/Developer/Concert/ConcertUI/ConcertSharedSlate/Source/ConcertSharedSlate/Private/Replication/Editor/View/ObjectViewer/Property/SPropertyTreeView.h b/Engine/Plugins/Developer/Concert/ConcertUI/ConcertSharedSlate/Source/ConcertSharedSlate/Private/Replication/Editor/View/ObjectViewer/Property/SPropertyTreeView.h
--- a/Engine/Plugins/Developer/Concert/ConcertUI/ConcertSharedSlate/Source/ConcertSharedSlate/Private/Replication/Editor/View/ObjectViewer/Property/SPropertyTreeView.h
+++ b/Engine/Plugins/Developer/Concert/ConcertUI/ConcertSharedSlate/Source/ConcertSharedSlate/Private/Replication/Editor/View/ObjectViewer/Property/SPropertyTreeView.h
@@ -11,6 +11,7 @@ namespace UE::ConcertSharedSlate
 {
 	class FReplicatedObjectData;
 	class FReplicatedPropertyData;
+	class IObjectNameModel;
 	class SReplicationFilterBar;
 	
 	/** Displays a searchable tree view of properties (SReplicationTreeView) and decorates it with a SBasicFilterBar. */
@@ -42,6 +43,9 @@ namespace UE::ConcertSharedSlate
 			/** How many items are to allowed to be selected */
 			SLATE_ARGUMENT(ESelectionMode::Type, SelectionMode)
 		
+			/** Optional. If set, this determines the display text for objects. */
+			SLATE_ARGUMENT(TSharedPtr<IObjectNameModel>, NameModel)
+		
 			/** Optional widget to add to the left of the search bar. */
 			SLATE_NAMED_SLOT(FArguments, LeftOfSearchBar)
 
@@ -69,6 +73,9 @@ namespace UE::ConcertSharedSlate
 
 		/** Used to tell the user that the selected object have all properties filtered out. */
 		TAttribute<TArray<FSoftObjectPath>> SelectedObjectsAttribute;
+		
+		/** Optional. If set, this determines the display text for objects. */
+		TSharedPtr<IObjectNameModel> NameModel;
 
 		struct FBuildFilterBarResult
 		{
diff --git a/Engine/Plugins/Developer/Concert/ConcertUI/ConcertSharedSlate/Source/ConcertSharedSlate/Private/Replication/Editor/View/ObjectViewer/SReplicatedPropertyView.cpp b/Engine/Plugins/Developer/Concert/ConcertUI/ConcertSharedSlate/Source/ConcertSharedSlate/Private/Replication/Editor/View/ObjectViewer/SReplicatedPropertyView.cpp
--- a/Engine/Plugins/Developer/Concert/ConcertUI/ConcertSharedSlate/Source/ConcertSharedSlate/Private/Replication/Editor/View/ObjectViewer/SReplicatedPropertyView.cpp
+++ b/Engine/Plugins/Developer/Concert/ConcertUI/ConcertSharedSlate/Source/ConcertSharedSlate/Private/Replication/Editor/View/ObjectViewer/SReplicatedPropertyView.cpp
@@ -129,6 +129,7 @@ namespace UE::ConcertSharedSlate
 				.PrimarySort(PrimarySort)
 				.SecondarySort(SecondarySort)
 				.SelectionMode(ESelectionMode::Multi)
+				.NameModel(InArgs._NameModel)
 				.LeftOfSearchBar()
 				[
 					InArgs._LeftOfPropertySearchBar.Widget
diff --git a/Engine/Plugins/Developer/Concert/ConcertUI/ConcertSharedSlate/Source/ConcertSharedSlate/Private/Replication/Editor/View/ObjectViewer/SReplicatedPropertyView.h b/Engine/Plugins/Developer/Concert/ConcertUI/ConcertSharedSlate/Source/ConcertSharedSlate/Private/Replication/Editor/View/ObjectViewer/SReplicatedPropertyView.h
--- a/Engine/Plugins/Developer/Concert/ConcertUI/ConcertSharedSlate/Source/ConcertSharedSlate/Private/Replication/Editor/View/ObjectViewer/SReplicatedPropertyView.h
+++ b/Engine/Plugins/Developer/Concert/ConcertUI/ConcertSharedSlate/Source/ConcertSharedSlate/Private/Replication/Editor/View/ObjectViewer/SReplicatedPropertyView.h
@@ -34,6 +34,9 @@ namespace UE::ConcertSharedSlate
 			/** Gets the root objects selected in the object outliner. */
 			SLATE_EVENT(FGetSelectedRootObjects, GetSelectedRootObjects)
 		
+			/** Optional. If set, this determines the display text for objects. */
+			SLATE_ARGUMENT(TSharedPtr<IObjectNameModel>, NameModel)
+		
 			/** Optional widget to add to the left of the property list search bar. */
 			SLATE_NAMED_SLOT(FArguments, LeftOfPropertySearchBar)
 		SLATE_END_ARGS()
diff --git a/Engine/Plugins/Developer/Concert/ConcertUI/ConcertSharedSlate/Source/ConcertSharedSlate/Private/Replication/Editor/View/ObjectViewer/SReplicationStreamViewer.cpp b/Engine/Plugins/Developer/Concert/ConcertUI/ConcertSharedSlate/Source/ConcertSharedSlate/Private/Replication/Editor/View/ObjectViewer/SReplicationStreamViewer.cpp
--- a/Engine/Plugins/Developer/Concert/ConcertUI/ConcertSharedSlate/Source/ConcertSharedSlate/Private/Replication/Editor/View/ObjectViewer/SReplicationStreamViewer.cpp
+++ b/Engine/Plugins/Developer/Concert/ConcertUI/ConcertSharedSlate/Source/ConcertSharedSlate/Private/Replication/Editor/View/ObjectViewer/SReplicationStreamViewer.cpp
@@ -24,6 +24,7 @@ namespace UE::ConcertSharedSlate
 	{
 		PropertiesModel = MoveTemp(InPropertiesModel);
 		SubobjectModel = InArgs._SubobjectModel;
+		NameModel = InArgs._NameModel;
 		
 		ChildSlot
 		[
@@ -174,7 +175,7 @@ namespace UE::ConcertSharedSlate
 	{
 		TArray Columns
 		{
-			ReplicationColumns::TopLevel::LabelColumn(PropertiesModel.ToSharedRef(), SubobjectModel.Get()),
+			ReplicationColumns::TopLevel::LabelColumn(PropertiesModel.ToSharedRef(), NameModel.Get()),
 			ReplicationColumns::TopLevel::TypeColumn(PropertiesModel.ToSharedRef())
 		};
 		Columns.Append(InArgs._AdditionalObjectColumns);
@@ -235,6 +236,7 @@ namespace UE::ConcertSharedSlate
 					.PrimarySort(InArgs._PrimaryPropertySort)
 					.SecondarySort(InArgs._SecondaryPropertySort)
 					.GetSelectedRootObjects_Lambda([this](){ return GetSelectedOutlinerObjects(); })
+					.NameModel(InArgs._NameModel)
 					.LeftOfPropertySearchBar()
 					[
 						InArgs._LeftOfPropertySearchBar.Widget
@@ -272,7 +274,7 @@ namespace UE::ConcertSharedSlate
 
 	void SReplicationStreamViewer::BuildObjectHierarchyIfNeeded(TSharedPtr<FReplicatedObjectData> ReplicatedObjectData, TMap<FSoftObjectPath, TSharedPtr<FReplicatedObjectData>>& NewPathToObjectDataCache)
 	{
-		// We're are supposed to display any hierarchy in the outliner if SubobjectModel is not set.
+		// We're are not supposed to display any hierarchy in the outliner if SubobjectModel is not set.
 		const FSoftObjectPath& ObjectPath = ReplicatedObjectData->GetObjectPath();
 		if (!SubobjectModel)
 		{
@@ -288,6 +290,7 @@ namespace UE::ConcertSharedSlate
 		const FSoftObjectPath TopLevelObject = OwningActor.Get(ObjectPath);
 		SubobjectModel->SetTopLevelObject(TopLevelObject);
 		
+		// Add all objects that appear in the hierarchy of ReplicatedObjectData
 		const auto AddItem = [this, &NewPathToObjectDataCache](const FSoftObjectPath& ObjectPath)
 		{
 			const TSharedPtr<FReplicatedObjectData>* ExistingItem = PathToObjectDataCache.Find(ObjectPath);
@@ -296,11 +299,10 @@ namespace UE::ConcertSharedSlate
 			AllObjectRowData.AddUnique(Item);
 			NewPathToObjectDataCache.Emplace(ObjectPath, Item);
 		};
-		// Add all objects that appear in the hierarchy of ReplicatedObjectData
 		AddItem(TopLevelObject);
 		SubobjectModel->ForEachSubobject([this, &AddItem](const FSoftObjectPath&, const FSoftObjectPath& ChildObject)
 		{
-			 AddItem(ChildObject);
+			AddItem(ChildObject);
 			return EBreakBehavior::Continue;
 		});
 	}
diff --git a/Engine/Plugins/Developer/Concert/ConcertUI/ConcertSharedSlate/Source/ConcertSharedSlate/Private/Replication/Editor/View/ObjectViewer/SReplicationStreamViewer.h b/Engine/Plugins/Developer/Concert/ConcertUI/ConcertSharedSlate/Source/ConcertSharedSlate/Private/Replication/Editor/View/ObjectViewer/SReplicationStreamViewer.h
--- a/Engine/Plugins/Developer/Concert/ConcertUI/ConcertSharedSlate/Source/ConcertSharedSlate/Private/Replication/Editor/View/ObjectViewer/SReplicationStreamViewer.h
+++ b/Engine/Plugins/Developer/Concert/ConcertUI/ConcertSharedSlate/Source/ConcertSharedSlate/Private/Replication/Editor/View/ObjectViewer/SReplicationStreamViewer.h
@@ -25,6 +25,7 @@ namespace UE::ConcertSharedSlate
 	class FReplicatedPropertyData;
 	class FReplicatedObjectData;
 	class IEditableReplicationStreamModel;
+	class IObjectNameModel;
 	class IReplicationStreamModel;
 	class ISubobjectModel;
 	class SPropertyTreeView;
@@ -61,6 +62,8 @@ namespace UE::ConcertSharedSlate
 
 			/** Optional. If set, this determines the children nested under the root objects. */
 			SLATE_ARGUMENT(TSharedPtr<ISubobjectModel>, SubobjectModel)
+			/** Optional. If set, this determines the display text for objects. */
+			SLATE_ARGUMENT(TSharedPtr<IObjectNameModel>, NameModel)
 
 			/** Optional. Called when the delete key is pressed in the object view. */
 			SLATE_EVENT(SReplicationTreeView<FReplicatedObjectData>::FDeleteItems, OnDeleteObjects)
@@ -101,8 +104,10 @@ namespace UE::ConcertSharedSlate
 
 		/** The model this view is visualizing. */
 		TSharedPtr<IReplicationStreamModel> PropertiesModel;
-		/** Can be null. If set, this determines the children nested under the root objects. */
+		/** Can be null. If set, this determines the children nested under the root objects. Editor builds have access to e.g. to USubobjectDataSubsystem but programs do not. */
 		TSharedPtr<ISubobjectModel> SubobjectModel;
+		/** Can be null. If set, this determines the display text for objects. Editor builds have access to e.g. to USubobjectDataSubsystem but programs do not. */
+		TSharedPtr<IObjectNameModel> NameModel;
 
 		/** Lists the properties of the selected actor */
 		TSharedPtr<SExpandableArea> PropertyArea;
diff --git a/Engine/Plugins/Developer/Concert/ConcertUI/ConcertSharedSlate/Source/ConcertSharedSlate/Private/Replication/Editor/View/SelectionViewerColumns.cpp b/Engine/Plugins/Developer/Concert/ConcertUI/ConcertSharedSlate/Source/ConcertSharedSlate/Private/Replication/Editor/View/SelectionViewerColumns.cpp
--- a/Engine/Plugins/Developer/Concert/ConcertUI/ConcertSharedSlate/Source/ConcertSharedSlate/Private/Replication/Editor/View/SelectionViewerColumns.cpp
+++ b/Engine/Plugins/Developer/Concert/ConcertUI/ConcertSharedSlate/Source/ConcertSharedSlate/Private/Replication/Editor/View/SelectionViewerColumns.cpp
@@ -3,7 +3,7 @@
 #include "Replication/Editor/View/SelectionViewerColumns.h"
 
 #include "Replication/Editor/Model/IEditableReplicationStreamModel.h"
-#include "Replication/Editor/Model/ISubobjectModel.h"
+#include "Replication/Editor/Model/Object/IObjectNameModel.h"
 #include "Replication/Editor/Model/ReplicatedPropertyData.h"
 #include "Replication/Editor/Model/ReplicatedObjectData.h"
 #include "Replication/Editor/View/DisplayUtils.h"
@@ -27,21 +27,12 @@ namespace UE::ConcertSharedSlate::ReplicationColumns::TopLevel
 	const FName LabelColumnId = TEXT("LabelColumn");
 	const FName TypeColumnId = TEXT("TypeColumn");
 	
-	FReplicationTopLevelObjectColumn LabelColumn(TSharedRef<IReplicationStreamModel> Model, ISubobjectModel* SubobjectModel)
+	FReplicationTopLevelObjectColumn LabelColumn(TSharedRef<IReplicationStreamModel> Model, IObjectNameModel* OptionalNameModel)
 	{
-		auto GetDisplayText = [SubobjectModel](const FReplicatedObjectData& ObjectData)
+		auto GetDisplayText = [OptionalNameModel](const FReplicatedObjectData& ObjectData)
 		{
 			const FSoftObjectPath& ObjectPath = ObjectData.GetObjectPath();
-			FText Text = DisplayUtils::GetObjectDisplayText(ObjectPath);
-			if (SubobjectModel && !SubobjectModel->IsTopLevelObject(ObjectPath))
-			{
-				if (const TOptional<FSoftObjectPath> OwningActorPath = ObjectUtils::GetActorOf(ObjectPath))
-				{
-					SubobjectModel->SetTopLevelObject(*OwningActorPath);
-					Text = SubobjectModel->GetSubobjectDisplayName(ObjectPath);
-				}
-			}
-			return Text;
+			return DisplayUtils::GetObjectDisplayText(ObjectPath, OptionalNameModel);
 		};
 		
 		return FReplicationTopLevelObjectColumn(
@@ -72,9 +63,9 @@ namespace UE::ConcertSharedSlate::ReplicationColumns::TopLevel
 							.Text(Text)
 						];
 				})
-				.PopulateSearchItems_Lambda([](const FReplicatedObjectData& ObjectData, TArray<FString>& InOutSearchStrings)
+				.PopulateSearchItems_Lambda([OptionalNameModel](const FReplicatedObjectData& ObjectData, TArray<FString>& InOutSearchStrings)
 				{
-					InOutSearchStrings.Add(DisplayUtils::GetObjectDisplayText(ObjectData.GetObjectPath()).ToString());
+					InOutSearchStrings.Add(DisplayUtils::GetObjectDisplayText(ObjectData.GetObjectPath(), OptionalNameModel).ToString());
 				})
 				.IsLessThan_Lambda([GetDisplayText](const FReplicatedObjectData& Left, const FReplicatedObjectData& Right)
 				{
diff --git a/Engine/Plugins/Developer/Concert/ConcertUI/ConcertSharedSlate/Source/ConcertSharedSlate/Private/Replication/ObjectUtils.cpp b/Engine/Plugins/Developer/Concert/ConcertUI/ConcertSharedSlate/Source/ConcertSharedSlate/Private/Replication/ObjectUtils.cpp
--- a/Engine/Plugins/Developer/Concert/ConcertUI/ConcertSharedSlate/Source/ConcertSharedSlate/Private/Replication/ObjectUtils.cpp
+++ b/Engine/Plugins/Developer/Concert/ConcertUI/ConcertSharedSlate/Source/ConcertSharedSlate/Private/Replication/ObjectUtils.cpp
@@ -55,5 +55,17 @@ namespace UE::ConcertSharedSlate::ObjectUtils
 		const FSoftObjectPath PathToOwningActor(SoftObjectPath.GetAssetPath(), NewSubstring);
 		return PathToOwningActor;
 	}
+	
+	FString ExtractObjectDisplayStringFromPath(const FSoftObjectPath& Object)
+	{
+		// Subpath looks like this PersistentLevel.Actor.Component
+		const FString& Subpath = Object.GetSubPathString();
+		const int32 LastDotIndex = Subpath.Find(TEXT("."), ESearchCase::CaseSensitive, ESearchDir::FromEnd);
+		if (LastDotIndex == INDEX_NONE)
+		{
+			return {};
+		}
+		return Subpath.RightChop(LastDotIndex + 1);
+	}
 };
 
diff --git a/Engine/Plugins/Developer/Concert/ConcertUI/ConcertSharedSlate/Source/ConcertSharedSlate/Private/Replication/ReplicationWidgetFactories.cpp b/Engine/Plugins/Developer/Concert/ConcertUI/ConcertSharedSlate/Source/ConcertSharedSlate/Private/Replication/ReplicationWidgetFactories.cpp
--- a/Engine/Plugins/Developer/Concert/ConcertUI/ConcertSharedSlate/Source/ConcertSharedSlate/Private/Replication/ReplicationWidgetFactories.cpp
+++ b/Engine/Plugins/Developer/Concert/ConcertUI/ConcertSharedSlate/Source/ConcertSharedSlate/Private/Replication/ReplicationWidgetFactories.cpp
@@ -26,6 +26,7 @@ namespace UE::ConcertSharedSlate
 			.PrimaryPropertySort(Params.ViewerParams.PrimaryPropertySort)
 			.SecondaryPropertySort(Params.ViewerParams.SecondaryPropertySort)
 			.SubobjectModel(Params.ViewerParams.SubobjectModel)
+			.NameModel(Params.ViewerParams.NameModel)
 			.OnExtendObjectsContextMenu(Params.ViewerParams.OnExtendObjectsContextMenu)
 			.LeftOfObjectSearchBar() [ Params.ViewerParams.LeftOfObjectSearchBar.Widget ]
 			.LeftOfPropertySearchBar() [ Params.ViewerParams.LeftOfPropertySearchBar.Widget ]
diff --git a/Engine/Plugins/Developer/Concert/ConcertUI/ConcertSharedSlate/Source/ConcertSharedSlate/Public/Replication/Editor/Model/ISubobjectModel.h b/Engine/Plugins/Developer/Concert/ConcertUI/ConcertSharedSlate/Source/ConcertSharedSlate/Public/Replication/Editor/Model/ISubobjectModel.h
--- a/Engine/Plugins/Developer/Concert/ConcertUI/ConcertSharedSlate/Source/ConcertSharedSlate/Public/Replication/Editor/Model/ISubobjectModel.h
+++ b/Engine/Plugins/Developer/Concert/ConcertUI/ConcertSharedSlate/Source/ConcertSharedSlate/Public/Replication/Editor/Model/ISubobjectModel.h
@@ -31,9 +31,6 @@ namespace UE::ConcertSharedSlate
 		 * Example: One category could be the the USceneComponent hierarchy and another could be all other UActorComponents (like in the SSubobjectEditor).
 		 */
 		virtual TArray<FName> GetCategories() const = 0;
-
-		/** Gets the display name for the subobject */
-		virtual FText GetSubobjectDisplayName(const FSoftObjectPath& ObjectPath) const = 0;
 		
 		/**
 		 * Gets the direct subobjects of the top level objects.
@@ -45,10 +42,6 @@ namespace UE::ConcertSharedSlate
 		/** Gets the direct subobject children of another subobject. Child subobjects are implicitly in the same category as Parent. */
 		virtual void ForEachDirectChildSubobject(const FSoftObjectPath& Parent, TFunctionRef<EBreakBehavior(const FSoftObjectPath& Object)> Callback) const = 0;
 
-		DECLARE_MULTICAST_DELEGATE(FOnHierarchyChanged)
-		/** Called when the hierarchy has changed, e.g. due to calling SetTopLevelObject. */
-		virtual FOnHierarchyChanged& OnHierarchyChanged() = 0;
-
 		/** Util for iterating all subobjects. */
 		void ForEachSubobject(TFunctionRef<EBreakBehavior(const FSoftObjectPath& Parent, const FSoftObjectPath& ChildObject)> Callback);
 
diff --git a/Engine/Plugins/Developer/Concert/ConcertUI/ConcertSharedSlate/Source/ConcertSharedSlate/Public/Replication/Editor/Model/Object/IObjectNameModel.h b/Engine/Plugins/Developer/Concert/ConcertUI/ConcertSharedSlate/Source/ConcertSharedSlate/Public/Replication/Editor/Model/Object/IObjectNameModel.h
new file mode 100644
--- /dev/null
+++ b/Engine/Plugins/Developer/Concert/ConcertUI/ConcertSharedSlate/Source/ConcertSharedSlate/Public/Replication/Editor/Model/Object/IObjectNameModel.h
@@ -0,0 +1,27 @@
+// Copyright Epic Games, Inc. All Rights Reserved.
+
+#pragma once
+
+#include "Delegates/Delegate.h"
+#include "Templates/SharedPointer.h"
+
+namespace UE::ConcertSharedSlate
+{
+	/**
+	 * Abstracts the concept of getting an object's display name.
+	 *
+	 * In editor builds, this could use the USubobjectDataSubsystem to look up subobject names on an actor.
+	 * In programs, it could just default to what it says in the object path.
+	 */
+	class CONCERTSHAREDSLATE_API IObjectNameModel
+	{
+	public:
+
+		/** @return Gets the display name for an object */
+		virtual FText GetObjectDisplayName(const FSoftObjectPath& ObjectPath) const = 0;
+
+		// In the future, we could add an event here for when an object's display name changes.
+
+		virtual ~IObjectNameModel() = default;
+	};
+}
\ No newline at end of file
diff --git a/Engine/Plugins/Developer/Concert/ConcertUI/ConcertSharedSlate/Source/ConcertSharedSlate/Public/Replication/Editor/View/SelectionViewerColumns.h b/Engine/Plugins/Developer/Concert/ConcertUI/ConcertSharedSlate/Source/ConcertSharedSlate/Public/Replication/Editor/View/SelectionViewerColumns.h
--- a/Engine/Plugins/Developer/Concert/ConcertUI/ConcertSharedSlate/Source/ConcertSharedSlate/Public/Replication/Editor/View/SelectionViewerColumns.h
+++ b/Engine/Plugins/Developer/Concert/ConcertUI/ConcertSharedSlate/Source/ConcertSharedSlate/Public/Replication/Editor/View/SelectionViewerColumns.h
@@ -11,7 +11,7 @@ namespace UE::ConcertSharedSlate
 	class IEditableReplicationStreamModel;
 	class IReplicationStreamViewer;
 	class IReplicationStreamModel;
-	class ISubobjectModel;
+	class IObjectNameModel;
 }
 
 namespace UE::ConcertSharedSlate::ReplicationColumns::TopLevel
@@ -28,7 +28,7 @@ namespace UE::ConcertSharedSlate::ReplicationColumns::TopLevel
 		Type = 30,
 	};
 
-	CONCERTSHAREDSLATE_API FReplicationTopLevelObjectColumn LabelColumn(TSharedRef<IReplicationStreamModel> Model, ISubobjectModel* SubobjectModel = nullptr);
+	CONCERTSHAREDSLATE_API FReplicationTopLevelObjectColumn LabelColumn(TSharedRef<IReplicationStreamModel> Model, IObjectNameModel* OptionalNameModel = nullptr);
 	CONCERTSHAREDSLATE_API FReplicationTopLevelObjectColumn TypeColumn(TSharedRef<IReplicationStreamModel> Model);
 }
 
diff --git a/Engine/Plugins/Developer/Concert/ConcertUI/ConcertSharedSlate/Source/ConcertSharedSlate/Public/Replication/ObjectUtils.h b/Engine/Plugins/Developer/Concert/ConcertUI/ConcertSharedSlate/Source/ConcertSharedSlate/Public/Replication/ObjectUtils.h
--- a/Engine/Plugins/Developer/Concert/ConcertUI/ConcertSharedSlate/Source/ConcertSharedSlate/Public/Replication/ObjectUtils.h
+++ b/Engine/Plugins/Developer/Concert/ConcertUI/ConcertSharedSlate/Source/ConcertSharedSlate/Public/Replication/ObjectUtils.h
@@ -13,4 +13,7 @@ namespace UE::ConcertSharedSlate::ObjectUtils
 	
 	/** @return Get the owning actor of Subobject. If Subobject is an actor, then this returns unset. */
 	CONCERTSHAREDSLATE_API TOptional<FSoftObjectPath> GetActorOf(const FSoftObjectPath& Subobject);
+	
+	/** @return Gets the last object name in the subpath. */
+	CONCERTSHAREDSLATE_API FString ExtractObjectDisplayStringFromPath(const FSoftObjectPath& Object);
 };
diff --git a/Engine/Plugins/Developer/Concert/ConcertUI/ConcertSharedSlate/Source/ConcertSharedSlate/Public/Replication/ReplicationWidgetFactories.h b/Engine/Plugins/Developer/Concert/ConcertUI/ConcertSharedSlate/Source/ConcertSharedSlate/Public/Replication/ReplicationWidgetFactories.h
--- a/Engine/Plugins/Developer/Concert/ConcertUI/ConcertSharedSlate/Source/ConcertSharedSlate/Public/Replication/ReplicationWidgetFactories.h
+++ b/Engine/Plugins/Developer/Concert/ConcertUI/ConcertSharedSlate/Source/ConcertSharedSlate/Public/Replication/ReplicationWidgetFactories.h
@@ -23,6 +23,7 @@ namespace UE::ConcertSharedSlate
 	class IMultiReplicationStreamEditor;
 	
 	class IReplicationStreamModel;
+	class IObjectNameModel;
 	class IObjectSelectionSourceModel;
 	class IReplicationStreamEditor;
 	class IReplicationStreamViewer;
@@ -54,6 +55,9 @@ namespace UE::ConcertSharedSlate
 		 */
 		TSharedPtr<ISubobjectModel> SubobjectModel;
 
+		/** Optional. Determines the display name of objects. */
+		TSharedPtr<IObjectNameModel> NameModel;
+
 		/** Called to generate the context menu for objects. This extends the options already generated by this widget and is called at the end. */
 		FExtendObjectMenu OnExtendObjectsContextMenu;
 		
