Refactor ISubobjectModel to IObjectHierarchyModel

diff --git a/Engine/Plugins/Developer/Concert/ConcertApp/MultiUserClient/Source/MultiUserClient/Private/Settings/FrequencyDefaultRules.h b/Engine/Plugins/Developer/Concert/ConcertApp/MultiUserClient/Source/MultiUserClient/Private/Settings/FrequencyDefaultRules.h
--- a/Engine/Plugins/Developer/Concert/ConcertApp/MultiUserClient/Source/MultiUserClient/Private/Settings/FrequencyDefaultRules.h
+++ b/Engine/Plugins/Developer/Concert/ConcertApp/MultiUserClient/Source/MultiUserClient/Private/Settings/FrequencyDefaultRules.h
@@ -59,7 +59,7 @@ struct FMultiUserFrequencyDefaultRules
 
 	/**
 	 * Default frequency settings that are used by objects when they do not match any specific rules in FrequencyRuleStack.
-	 * These are the settings that are set for the Multi User stream's FSharedReplicationStreamDescription::FrequencySettings::Defaults member.
+	 * These are the settings that are set for the Multi User stream's FConcertBaseStreamInfo::FrequencySettings::Defaults member.
 	 */
 	UPROPERTY(Category = "Replication Settings|Frequency", EditAnywhere)
 	FConcertObjectReplicationSettings DefaultObjectFrequencySettings { EConcertObjectReplicationMode::SpecifiedRate, 30 };
@@ -84,7 +84,7 @@ struct FMultiUserFrequencyDefaultRules
 	 * Goes through FrequencyRuleStack and returns the FConcertObjectReplicationSettings of the entry that was matched.
 	 * 
 	 * If none was matched, the object should not have any frequency override and should default to the stream's
-	 * FSharedReplicationStreamDescription::FrequencySettings::Defaults, which should be set to DefaultObjectFrequencySettings. 
+	 * FConcertBaseStreamInfo::FrequencySettings::Defaults, which should be set to DefaultObjectFrequencySettings. 
 	 */
 	TOptional<FConcertObjectReplicationSettings> FindOverrideSettings(UObject& Object) const;
 };
diff --git a/Engine/Plugins/Developer/Concert/ConcertApp/MultiUserClient/Source/MultiUserClient/Private/Widgets/ActiveSession/Replication/Client/Multi/SMultiClientView.cpp b/Engine/Plugins/Developer/Concert/ConcertApp/MultiUserClient/Source/MultiUserClient/Private/Widgets/ActiveSession/Replication/Client/Multi/SMultiClientView.cpp
--- a/Engine/Plugins/Developer/Concert/ConcertApp/MultiUserClient/Source/MultiUserClient/Private/Widgets/ActiveSession/Replication/Client/Multi/SMultiClientView.cpp
+++ b/Engine/Plugins/Developer/Concert/ConcertApp/MultiUserClient/Source/MultiUserClient/Private/Widgets/ActiveSession/Replication/Client/Multi/SMultiClientView.cpp
@@ -92,7 +92,7 @@ namespace UE::MultiUserClient
 			.GetAutoAssignToStreamDelegate = MoveTemp(GetAutoAssignTargetDelegate),
 			.ViewerParams 
 			{
-				.SubobjectModel = ConcertClientSharedSlate::CreateSubobjectModelForComponentHierarchy(), // This makes actors have children in the top view
+				.ObjectHierarchy = ConcertClientSharedSlate::CreateObjectHierarchyForComponentHierarchy(), // This makes actors have children in the top view
 				.NameModel = ConcertClientSharedSlate::CreateEditorObjectNameModel(), // This makes actors use their labels, and components use the names given in the BP editor
 				.OnExtendObjectsContextMenu = FExtendObjectMenu::CreateSP(this, &SMultiClientView::ExtendObjectContextMenu),
 				.AdditionalObjectColumns =
diff --git a/Engine/Plugins/Developer/Concert/ConcertApp/MultiUserClient/Source/MultiUserClient/Private/Widgets/ActiveSession/Replication/Client/Single/SReplicationClientView.cpp b/Engine/Plugins/Developer/Concert/ConcertApp/MultiUserClient/Source/MultiUserClient/Private/Widgets/ActiveSession/Replication/Client/Single/SReplicationClientView.cpp
--- a/Engine/Plugins/Developer/Concert/ConcertApp/MultiUserClient/Source/MultiUserClient/Private/Widgets/ActiveSession/Replication/Client/Single/SReplicationClientView.cpp
+++ b/Engine/Plugins/Developer/Concert/ConcertApp/MultiUserClient/Source/MultiUserClient/Private/Widgets/ActiveSession/Replication/Client/Single/SReplicationClientView.cpp
@@ -78,7 +78,7 @@ namespace UE::MultiUserClient
 			.EditingDisabledToolTipText = LOCTEXT("Editing.NotImplemented", "Editing remote clients is not implemented. You can only edit the local client."),
 			.ViewerParams =
 			{
-				.SubobjectModel = ConcertClientSharedSlate::CreateSubobjectModelForComponentHierarchy(), // This makes actors have children in the top view
+				.ObjectHierarchy = ConcertClientSharedSlate::CreateObjectHierarchyForComponentHierarchy(), // This makes actors have children in the top view
 				.NameModel = ConcertClientSharedSlate::CreateEditorObjectNameModel(), // This makes actors use their labels, and components use the names given in the BP editor
 				.OnExtendObjectsContextMenu = FExtendObjectMenu::CreateSP(this, &SReplicationClientView::ExtendObjectContextMenu),
 				.AdditionalObjectColumns =
diff --git a/Engine/Plugins/Developer/Concert/ConcertUI/ConcertClientSharedSlate/Source/ConcertClientSharedSlate/Private/Replication/ClientReplicationWidgetFactories.cpp b/Engine/Plugins/Developer/Concert/ConcertUI/ConcertClientSharedSlate/Source/ConcertClientSharedSlate/Private/Replication/ClientReplicationWidgetFactories.cpp
--- a/Engine/Plugins/Developer/Concert/ConcertUI/ConcertClientSharedSlate/Source/ConcertClientSharedSlate/Private/Replication/ClientReplicationWidgetFactories.cpp
+++ b/Engine/Plugins/Developer/Concert/ConcertUI/ConcertClientSharedSlate/Source/ConcertClientSharedSlate/Private/Replication/ClientReplicationWidgetFactories.cpp
@@ -2,10 +2,10 @@
 
 #include "Replication/ClientReplicationWidgetFactories.h"
 
-#include "Object/EditorObjectNameModel.h"
 #include "Replication/ReplicationWidgetFactories.h"
+#include "Replication/Editor/Model/Object/EditorObjectHierarchyModel.h"
+#include "Replication/Editor/Model/Object/EditorObjectNameModel.h"
 #include "Replication/Editor/Model/ReplicationStreamObject.h"
-#include "Replication/Editor/Model/Subobject/ComponentHierarchySubobjectModel.h"
 #include "Replication/Editor/Model/TransactionalReplicationStreamModel.h"
 #include "Replication/Editor/View/ObjectEditor/SDefaultReplicationStreamEditor.h"
 
@@ -14,9 +14,9 @@
 
 namespace UE::ConcertClientSharedSlate
 {
-	TSharedRef<ConcertSharedSlate::ISubobjectModel> CreateSubobjectModelForComponentHierarchy()
+	TSharedRef<ConcertSharedSlate::IObjectHierarchyModel> CreateObjectHierarchyForComponentHierarchy()
 	{
-		return MakeShared<FComponentHierarchySubobjectModel>();
+		return MakeShared<FEditorObjectHierarchyModel>();
 	}
 
 	TSharedRef<ConcertSharedSlate::IObjectNameModel> CreateEditorObjectNameModel()
diff --git a/Engine/Plugins/Developer/Concert/ConcertUI/ConcertClientSharedSlate/Source/ConcertClientSharedSlate/Private/Replication/Editor/Model/Object/EditorObjectHierarchyModel.cpp b/Engine/Plugins/Developer/Concert/ConcertUI/ConcertClientSharedSlate/Source/ConcertClientSharedSlate/Private/Replication/Editor/Model/Object/EditorObjectHierarchyModel.cpp
new file mode 100644
--- /dev/null
+++ b/Engine/Plugins/Developer/Concert/ConcertUI/ConcertClientSharedSlate/Source/ConcertClientSharedSlate/Private/Replication/Editor/Model/Object/EditorObjectHierarchyModel.cpp
@@ -0,0 +1,95 @@
+// Copyright Epic Games, Inc. All Rights Reserved.
+
+#include "EditorObjectHierarchyModel.h"
+
+#include "Components/SceneComponent.h"
+#include "Misc/EBreakBehavior.h"
+
+#include "Components/SceneComponent.h"
+#include "GameFramework/Actor.h"
+#include "UObject/UObjectHash.h"
+
+namespace UE::ConcertClientSharedSlate
+{
+	namespace Private
+	{
+		static bool IsValidComponent(const UActorComponent& Component)
+		{
+			return !Component.IsVisualizationComponent();
+		}
+
+		static void DiscoverComponents(
+			const UObject& Object,
+			const TFunctionRef<EBreakBehavior(const FSoftObjectPath& Child, ConcertSharedSlate::EChildRelationship Relationship)>& Callback
+			)
+		{
+			using namespace ConcertSharedSlate;
+			
+			if (const AActor* AsActor = Cast<AActor>(&Object))
+			{
+				// Only the root component should be listed. It's children should be returned through recursive ForEachDirectChild calls.
+				USceneComponent* RootComponent = AsActor->GetRootComponent();
+				if (RootComponent && Private::IsValidComponent(*RootComponent) && Callback(RootComponent, EChildRelationship::Component) == EBreakBehavior::Break)
+				{
+					return;
+				}
+
+				for (const UActorComponent* Component : TInlineComponentArray<UActorComponent*>(AsActor))
+				{
+					if (IsValidComponent(*Component)
+						// The only component to be found here should be the root component, the others are found via recursive ForEachDirectChild calls
+						&& !Component->IsA<USceneComponent>()
+						&& Callback(Component, EChildRelationship::Component) == EBreakBehavior::Break)
+					{
+						break;
+					}
+				}
+			}
+		
+			if (const USceneComponent* SceneComponent = Cast<USceneComponent>(&Object))
+			{
+				for (int32 i = 0; i < SceneComponent->GetNumChildrenComponents(); ++i)
+				{
+					const USceneComponent* ChildComponent = SceneComponent->GetChildComponent(i);
+					if (IsValidComponent(*ChildComponent)
+						&& Callback(ChildComponent, EChildRelationship::Component) == EBreakBehavior::Break)
+					{
+						break;
+					}
+				}
+			}
+		}
+	}
+
+	void FEditorObjectHierarchyModel::ForEachDirectChild(
+		const FSoftObjectPath& Parent,
+		TFunctionRef<EBreakBehavior(const FSoftObjectPath& Object, ConcertSharedSlate::EChildRelationship Relationship)> Callback,
+		ConcertSharedSlate::EChildRelationshipFlags InclusionFlags
+		) const
+	{
+		using namespace ConcertSharedSlate;
+
+		// In the context of replication, only resolve objects but do not load them. Case: client 1 may be viewing objects of client 2 but they are in different worlds.
+		const UObject* Object = Parent.ResolveObject();
+		if (!Object)
+		{
+			return;
+		}
+
+		if (EnumHasAnyFlags(InclusionFlags, EChildRelationshipFlags::Component))
+		{
+			Private::DiscoverComponents(*Object, Callback);
+		}
+		
+		if (EnumHasAnyFlags(InclusionFlags, EChildRelationshipFlags::Subobject))
+		{
+			constexpr bool bIncludeNestedSubobjects = false;
+			ForEachObjectWithOuterBreakable(Object, [&Callback](UObject* Subobject)
+			{
+				// Components were already handled above - do not double report.
+				const bool bContinueIteration = Subobject->IsA<UActorComponent>() || Callback(Subobject, EChildRelationship::Subobject) == EBreakBehavior::Continue;
+				return bContinueIteration;
+			}, bIncludeNestedSubobjects);
+		}
+	}
+}
\ No newline at end of file
diff --git a/Engine/Plugins/Developer/Concert/ConcertUI/ConcertClientSharedSlate/Source/ConcertClientSharedSlate/Private/Replication/Editor/Model/Object/EditorObjectHierarchyModel.h b/Engine/Plugins/Developer/Concert/ConcertUI/ConcertClientSharedSlate/Source/ConcertClientSharedSlate/Private/Replication/Editor/Model/Object/EditorObjectHierarchyModel.h
new file mode 100644
--- /dev/null
+++ b/Engine/Plugins/Developer/Concert/ConcertUI/ConcertClientSharedSlate/Source/ConcertClientSharedSlate/Private/Replication/Editor/Model/Object/EditorObjectHierarchyModel.h
@@ -0,0 +1,34 @@
+// Copyright Epic Games, Inc. All Rights Reserved.
+
+#pragma once
+
+#include "Replication/Editor/Model/Object/IObjectHierarchyModel.h"
+#include "UObject/SoftObjectPtr.h"
+
+class AActor;
+
+namespace UE::ConcertClientSharedSlate
+{
+	/**
+	 * Builds a similar tree hierarchy as SSubobjectEditor with the addition of non-component subobjects.
+	 *
+	 * An example for a hierarchy in an editor build is:
+	 * - Actor returns its root component, non-scene components (UActorComponents), and all other direct subobjects
+	 * - Scene components return their children, and all other direct subobjects
+	 * - All other subobjects report all direct subobjects
+	 */
+	class FEditorObjectHierarchyModel : public ConcertSharedSlate::IObjectHierarchyModel
+	{
+	public:
+
+		//~ Begin IObjectHierarchy Interface
+		virtual void ForEachDirectChild(
+			const FSoftObjectPath& Parent,
+			TFunctionRef<EBreakBehavior(const FSoftObjectPath& Object, ConcertSharedSlate::EChildRelationship Relationship)> Callback,
+			ConcertSharedSlate::EChildRelationshipFlags InclusionFlags
+			) const override;
+		//~ End IObjectHierarchy Interface
+	};
+}
+
+
diff --git a/Engine/Plugins/Developer/Concert/ConcertUI/ConcertClientSharedSlate/Source/ConcertClientSharedSlate/Private/Replication/Editor/Model/Object/EditorObjectNameModel.cpp b/Engine/Plugins/Developer/Concert/ConcertUI/ConcertClientSharedSlate/Source/ConcertClientSharedSlate/Private/Replication/Editor/Model/Object/EditorObjectNameModel.cpp
new file mode 100644
--- /dev/null
+++ b/Engine/Plugins/Developer/Concert/ConcertUI/ConcertClientSharedSlate/Source/ConcertClientSharedSlate/Private/Replication/Editor/Model/Object/EditorObjectNameModel.cpp
@@ -0,0 +1,58 @@
+// Copyright Epic Games, Inc. All Rights Reserved.
+
+#include "EditorObjectNameModel.h"
+
+#include "Internationalization/Text.h"
+#include "Replication/ObjectUtils.h"
+
+#include "GameFramework/Actor.h"
+#include "SubobjectDataSubsystem.h"
+
+namespace UE::ConcertClientSharedSlate
+{
+	namespace Private
+	{
+		static FText FindSubobjectDisplayName(const UObject& Subbject, AActor& OwningActor)
+		{
+			USubobjectDataSubsystem* SubobjectDataSubsystem = USubobjectDataSubsystem::Get();
+			TArray<FSubobjectDataHandle> Handles;
+			SubobjectDataSubsystem->GatherSubobjectData(&OwningActor, Handles);
+				
+			for (const FSubobjectDataHandle& Handle : Handles)
+			{
+				const FSubobjectData* SubobjectData = Handle.GetData();
+				const UObject* Object = SubobjectData->FindComponentInstanceInActor(&OwningActor);
+				if (Object == &Subbject)
+				{
+					return FText::FromString(SubobjectData->GetDisplayString());
+				}
+			}
+
+			return FText::GetEmpty();
+		}
+	}
+	
+	FText FEditorObjectNameModel::GetObjectDisplayName(const FSoftObjectPath& ObjectPath) const
+	{
+		if (UObject* ResolvedObject = ObjectPath.ResolveObject())
+		{
+			// Display actor just like the outliner does
+			if (const AActor* AsActor = Cast<AActor>(ResolvedObject))
+			{
+				return FText::FromString(AsActor->GetActorLabel());
+			}
+
+			// Display the same component name as the SSubobjectEditor widget does, i.e. component hierarchy in the details panel or Blueprint editor.
+			AActor* OwningActor = ResolvedObject->GetTypedOuter<AActor>();
+			const FText FoundSubobjectName = OwningActor ? Private::FindSubobjectDisplayName(*ResolvedObject, *OwningActor) : FText::GetEmpty();
+			if (!FoundSubobjectName.IsEmpty())
+			{
+				return FoundSubobjectName;
+			}
+
+			return FText::FromString(ResolvedObject->GetName());
+		}
+		
+		return FText::FromString(ConcertSharedSlate::ObjectUtils::ExtractObjectDisplayStringFromPath(ObjectPath));
+	}
+}
diff --git a/Engine/Plugins/Developer/Concert/ConcertUI/ConcertClientSharedSlate/Source/ConcertClientSharedSlate/Private/Replication/Editor/Model/Object/EditorObjectNameModel.h b/Engine/Plugins/Developer/Concert/ConcertUI/ConcertClientSharedSlate/Source/ConcertClientSharedSlate/Private/Replication/Editor/Model/Object/EditorObjectNameModel.h
new file mode 100644
--- /dev/null
+++ b/Engine/Plugins/Developer/Concert/ConcertUI/ConcertClientSharedSlate/Source/ConcertClientSharedSlate/Private/Replication/Editor/Model/Object/EditorObjectNameModel.h
@@ -0,0 +1,19 @@
+// Copyright Epic Games, Inc. All Rights Reserved.
+
+#pragma once
+
+#include "Replication/Editor/Model/Object/IObjectNameModel.h"
+
+namespace UE::ConcertClientSharedSlate
+{
+	/** Name model that uses editor data for determining display names: actors use their labels, components ask USubobjectDataSubsystem.  */
+	class FEditorObjectNameModel : public ConcertSharedSlate::IObjectNameModel
+	{
+	public:
+
+		//~ Begin IObjectNameModel Interface
+		virtual FText GetObjectDisplayName(const FSoftObjectPath& ObjectPath) const override;
+		//~ End IObjectNameModel Interface
+	};
+}
+
diff --git a/Engine/Plugins/Developer/Concert/ConcertUI/ConcertClientSharedSlate/Source/ConcertClientSharedSlate/Private/Replication/Editor/Model/Subobject/ComponentHierarchySubobjectModel.cpp b/Engine/Plugins/Developer/Concert/ConcertUI/ConcertClientSharedSlate/Source/ConcertClientSharedSlate/Private/Replication/Editor/Model/Subobject/ComponentHierarchySubobjectModel.cpp
deleted file mode 100644
diff --git a/Engine/Plugins/Developer/Concert/ConcertUI/ConcertClientSharedSlate/Source/ConcertClientSharedSlate/Private/Replication/Editor/Model/Subobject/ComponentHierarchySubobjectModel.h b/Engine/Plugins/Developer/Concert/ConcertUI/ConcertClientSharedSlate/Source/ConcertClientSharedSlate/Private/Replication/Editor/Model/Subobject/ComponentHierarchySubobjectModel.h
deleted file mode 100644
diff --git a/Engine/Plugins/Developer/Concert/ConcertUI/ConcertClientSharedSlate/Source/ConcertClientSharedSlate/Private/Replication/Editor/View/ClientEditorColumns.h b/Engine/Plugins/Developer/Concert/ConcertUI/ConcertClientSharedSlate/Source/ConcertClientSharedSlate/Private/Replication/Editor/View/ClientEditorColumns.h
--- a/Engine/Plugins/Developer/Concert/ConcertUI/ConcertClientSharedSlate/Source/ConcertClientSharedSlate/Private/Replication/Editor/View/ClientEditorColumns.h
+++ b/Engine/Plugins/Developer/Concert/ConcertUI/ConcertClientSharedSlate/Source/ConcertClientSharedSlate/Private/Replication/Editor/View/ClientEditorColumns.h
@@ -11,7 +11,7 @@ namespace UE::ConcertSharedSlate
 	class IEditableReplicationStreamModel;
 	class IReplicationStreamViewer;
 	class IReplicationStreamModel;
-	class ISubobjectModel;
+	class IObjectHierarchyModel;
 }
 
 namespace UE::ConcertClientSharedSlate::ReplicationColumns::Property
diff --git a/Engine/Plugins/Developer/Concert/ConcertUI/ConcertClientSharedSlate/Source/ConcertClientSharedSlate/Private/Replication/Object/EditorObjectNameModel.cpp b/Engine/Plugins/Developer/Concert/ConcertUI/ConcertClientSharedSlate/Source/ConcertClientSharedSlate/Private/Replication/Object/EditorObjectNameModel.cpp
deleted file mode 100644
diff --git a/Engine/Plugins/Developer/Concert/ConcertUI/ConcertClientSharedSlate/Source/ConcertClientSharedSlate/Private/Replication/Object/EditorObjectNameModel.h b/Engine/Plugins/Developer/Concert/ConcertUI/ConcertClientSharedSlate/Source/ConcertClientSharedSlate/Private/Replication/Object/EditorObjectNameModel.h
deleted file mode 100644
diff --git a/Engine/Plugins/Developer/Concert/ConcertUI/ConcertClientSharedSlate/Source/ConcertClientSharedSlate/Public/Replication/ClientReplicationWidgetFactories.h b/Engine/Plugins/Developer/Concert/ConcertUI/ConcertClientSharedSlate/Source/ConcertClientSharedSlate/Public/Replication/ClientReplicationWidgetFactories.h
--- a/Engine/Plugins/Developer/Concert/ConcertUI/ConcertClientSharedSlate/Source/ConcertClientSharedSlate/Public/Replication/ClientReplicationWidgetFactories.h
+++ b/Engine/Plugins/Developer/Concert/ConcertUI/ConcertClientSharedSlate/Source/ConcertClientSharedSlate/Public/Replication/ClientReplicationWidgetFactories.h
@@ -13,14 +13,14 @@ namespace UE::ConcertSharedSlate
 	class IObjectNameModel;
 	class IReplicationStreamEditor;
 	class IStreamExtender;
-	class ISubobjectModel;
+	class IObjectHierarchyModel;
 	struct FCreateEditorParams;
 }
 
 namespace UE::ConcertClientSharedSlate
 {
 	/** Builds a similar tree hierarchy as SSubobjectEditor. Reports only components as subobjects. */
-	CONCERTCLIENTSHAREDSLATE_API TSharedRef<ConcertSharedSlate::ISubobjectModel> CreateSubobjectModelForComponentHierarchy();
+	CONCERTCLIENTSHAREDSLATE_API TSharedRef<ConcertSharedSlate::IObjectHierarchyModel> CreateObjectHierarchyForComponentHierarchy();
 
 	/** Name model that uses editor data for determining display names: actors use their labels, components ask USubobjectDataSubsystem. */
 	CONCERTCLIENTSHAREDSLATE_API TSharedRef<ConcertSharedSlate::IObjectNameModel> CreateEditorObjectNameModel();
diff --git a/Engine/Plugins/Developer/Concert/ConcertUI/ConcertSharedSlate/Source/ConcertSharedSlate/Private/Replication/Editor/Model/ISubobjectModel.cpp b/Engine/Plugins/Developer/Concert/ConcertUI/ConcertSharedSlate/Source/ConcertSharedSlate/Private/Replication/Editor/Model/ISubobjectModel.cpp
deleted file mode 100644
diff --git a/Engine/Plugins/Developer/Concert/ConcertUI/ConcertSharedSlate/Source/ConcertSharedSlate/Private/Replication/Editor/Model/Object/IObjectHierarchyModel.cpp b/Engine/Plugins/Developer/Concert/ConcertUI/ConcertSharedSlate/Source/ConcertSharedSlate/Private/Replication/Editor/Model/Object/IObjectHierarchyModel.cpp
new file mode 100644
--- /dev/null
+++ b/Engine/Plugins/Developer/Concert/ConcertUI/ConcertSharedSlate/Source/ConcertSharedSlate/Private/Replication/Editor/Model/Object/IObjectHierarchyModel.cpp
@@ -0,0 +1,46 @@
+// Copyright Epic Games, Inc. All Rights Reserved.
+
+#include "Replication/Editor/Model/Object/IObjectHierarchyModel.h"
+
+#include "UObject/SoftObjectPath.h"
+
+namespace UE::ConcertSharedSlate
+{
+	void IObjectHierarchyModel::ForEachChildRecursive(
+		const FSoftObjectPath& Root,
+		TFunctionRef<EBreakBehavior(const FSoftObjectPath& Parent, const FSoftObjectPath& ChildObject, EChildRelationship Relationship)> Callback,
+		EChildRelationshipFlags InclusionFlags
+		)
+	{
+		struct FHelper
+		{
+			static EBreakBehavior Visit(
+				IObjectHierarchyModel& Model,
+				const FSoftObjectPath& Parent,
+				const FSoftObjectPath& ChildObject,
+				EChildRelationship Relationship,
+				TFunctionRef<EBreakBehavior(const FSoftObjectPath&, const FSoftObjectPath&, EChildRelationship)> Callback,
+				EChildRelationshipFlags InclusionFlags
+				)
+			{
+				if (Callback(Parent, ChildObject, Relationship) == EBreakBehavior::Break)
+				{
+					return EBreakBehavior::Break;
+				}
+				
+				EBreakBehavior Result = EBreakBehavior::Continue;
+				Model.ForEachDirectChild(ChildObject, [&Model, &ChildObject, &Callback, &InclusionFlags, &Result](const FSoftObjectPath& ChildOfChild, EChildRelationship Relationship)
+				{
+					Result = Visit(Model, ChildObject, ChildOfChild, Relationship, Callback, InclusionFlags);
+					return Result;
+				});
+				return Result;
+			}
+		};
+			
+		ForEachDirectChild(Root, [this, &Root, &Callback, &InclusionFlags](const FSoftObjectPath& ChildObject, EChildRelationship Relationship)
+		{
+			return FHelper::Visit(*this, Root, ChildObject, Relationship, Callback, InclusionFlags);
+		}, InclusionFlags);
+	}
+}
\ No newline at end of file
diff --git a/Engine/Plugins/Developer/Concert/ConcertUI/ConcertSharedSlate/Source/ConcertSharedSlate/Private/Replication/Editor/View/ObjectEditor/SBaseReplicationStreamEditor.cpp b/Engine/Plugins/Developer/Concert/ConcertUI/ConcertSharedSlate/Source/ConcertSharedSlate/Private/Replication/Editor/View/ObjectEditor/SBaseReplicationStreamEditor.cpp
--- a/Engine/Plugins/Developer/Concert/ConcertUI/ConcertSharedSlate/Source/ConcertSharedSlate/Private/Replication/Editor/View/ObjectEditor/SBaseReplicationStreamEditor.cpp
+++ b/Engine/Plugins/Developer/Concert/ConcertUI/ConcertSharedSlate/Source/ConcertSharedSlate/Private/Replication/Editor/View/ObjectEditor/SBaseReplicationStreamEditor.cpp
@@ -12,7 +12,7 @@
 #include "Replication/Editor/Model/ObjectSource/IObjectSelectionSourceModel.h"
 
 #include "Algo/AnyOf.h"
-#include "Replication/Editor/Model/ISubobjectModel.h"
+#include "Replication/Editor/Model/Object/IObjectHierarchyModel.h"
 #include "UObject/Class.h"
 #include "Widgets/SBoxPanel.h"
 
@@ -28,7 +28,7 @@ namespace UE::ConcertSharedSlate
 	{
 		ObjectSelectionSource = MoveTemp(InObjectSelectionSource);
 		PropertySelectionSource = MoveTemp(InPropertySelectionSource);
-		SubobjectModel = InArgs._SubobjectModel;
+		ObjectHierarchy = InArgs._ObjectHierarchy;
 		
 		EditablePropertiesModel = MoveTemp(InPropertiesModel);
 		EditablePropertiesModel->OnObjectsChanged().AddSP(this, &SBaseReplicationStreamEditor::OnObjectsChanged);
@@ -49,7 +49,7 @@ namespace UE::ConcertSharedSlate
 				.AdditionalPropertyColumns(InArgs._AdditionalPropertyColumns)
 				.PrimaryPropertySort(InArgs._PrimaryPropertySort)
 				.SecondaryPropertySort(InArgs._SecondaryPropertySort)
-				.SubobjectModel(InArgs._SubobjectModel)
+				.ObjectHierarchy(InArgs._ObjectHierarchy)
 				.NameModel(InArgs._NameModel)
 				.OnDeleteObjects(this, &SBaseReplicationStreamEditor::OnDeleteObjects)
 				.OnObjectsContextMenuOpening(this, &SBaseReplicationStreamEditor::OnObjectsContextMenuOpening)
diff --git a/Engine/Plugins/Developer/Concert/ConcertUI/ConcertSharedSlate/Source/ConcertSharedSlate/Private/Replication/Editor/View/ObjectEditor/SBaseReplicationStreamEditor.h b/Engine/Plugins/Developer/Concert/ConcertUI/ConcertSharedSlate/Source/ConcertSharedSlate/Private/Replication/Editor/View/ObjectEditor/SBaseReplicationStreamEditor.h
--- a/Engine/Plugins/Developer/Concert/ConcertUI/ConcertSharedSlate/Source/ConcertSharedSlate/Private/Replication/Editor/View/ObjectEditor/SBaseReplicationStreamEditor.h
+++ b/Engine/Plugins/Developer/Concert/ConcertUI/ConcertSharedSlate/Source/ConcertSharedSlate/Private/Replication/Editor/View/ObjectEditor/SBaseReplicationStreamEditor.h
@@ -25,7 +25,7 @@ namespace UE::ConcertSharedSlate
 	class IPropertySelectionSourceModel;
 	class IReplicationStreamModel;
 	class IReplicationSubobjectView;
-	class ISubobjectModel;
+	class IObjectHierarchyModel;
 	class SReplicationStreamViewer;
 	
 	enum class EReplicatedObjectChangeReason : uint8;
@@ -67,7 +67,7 @@ namespace UE::ConcertSharedSlate
 			SLATE_ARGUMENT(FColumnSortInfo, SecondaryPropertySort)
 		
 			/** Optional. If set, this determines the children nested under the root objects. */
-			SLATE_ARGUMENT(TSharedPtr<ISubobjectModel>, SubobjectModel)
+			SLATE_ARGUMENT(TSharedPtr<IObjectHierarchyModel>, ObjectHierarchy)
 			/** Optional. If set, this determines the display text for objects. */
 			SLATE_ARGUMENT(TSharedPtr<IObjectNameModel>, NameModel)
 
@@ -112,7 +112,7 @@ namespace UE::ConcertSharedSlate
 		 */
 		TSharedPtr<FFakeObjectToPropertiesEditorModel> PropertiesModelAdapter;
 		/** Can be null. If set, adds all subobjects to the top level view. */
-		TSharedPtr<ISubobjectModel> SubobjectModel;
+		TSharedPtr<IObjectHierarchyModel> ObjectHierarchy;
 		
 		/** For deciding which objects can be added to EditablePropertiesModel. */
 		TSharedPtr<IObjectSelectionSourceModel> ObjectSelectionSource;
diff --git a/Engine/Plugins/Developer/Concert/ConcertUI/ConcertSharedSlate/Source/ConcertSharedSlate/Private/Replication/Editor/View/ObjectViewer/SReplicationStreamViewer.cpp b/Engine/Plugins/Developer/Concert/ConcertUI/ConcertSharedSlate/Source/ConcertSharedSlate/Private/Replication/Editor/View/ObjectViewer/SReplicationStreamViewer.cpp
--- a/Engine/Plugins/Developer/Concert/ConcertUI/ConcertSharedSlate/Source/ConcertSharedSlate/Private/Replication/Editor/View/ObjectViewer/SReplicationStreamViewer.cpp
+++ b/Engine/Plugins/Developer/Concert/ConcertUI/ConcertSharedSlate/Source/ConcertSharedSlate/Private/Replication/Editor/View/ObjectViewer/SReplicationStreamViewer.cpp
@@ -4,7 +4,7 @@
 
 #include "ConcertFrontendUtils.h"
 #include "Replication/Editor/Model/IReplicationStreamModel.h"
-#include "Replication/Editor/Model/ISubobjectModel.h"
+#include "Replication/Editor/Model/Object/IObjectHierarchyModel.h"
 #include "Replication/Editor/Model/ReplicatedObjectData.h"
 #include "Replication/Editor/View/ObjectViewer/Property/SPropertyTreeView.h"
 #include "Replication/Editor/View/SelectionViewerColumns.h"
@@ -23,7 +23,7 @@ namespace UE::ConcertSharedSlate
 	void SReplicationStreamViewer::Construct(const FArguments& InArgs, TSharedRef<IReplicationStreamModel> InPropertiesModel)
 	{
 		PropertiesModel = MoveTemp(InPropertiesModel);
-		SubobjectModel = InArgs._SubobjectModel;
+		ObjectHierarchy = InArgs._ObjectHierarchy;
 		NameModel = InArgs._NameModel;
 		
 		ChildSlot
@@ -128,9 +128,9 @@ namespace UE::ConcertSharedSlate
 				ItemsToExpand.Add(*Item);
 			}
 			
-			if (bRecursive && SubobjectModel)
+			if (bRecursive && ObjectHierarchy)
 			{
-				SubobjectModel->ForEachSubobject([this, &ItemsToExpand](const FSoftObjectPath&, const FSoftObjectPath& ChildObject)
+				ObjectHierarchy->ForEachChildRecursive(Path, [this, &ItemsToExpand](const FSoftObjectPath&, const FSoftObjectPath& ChildObject, EChildRelationship)
 				{
 					if (const TSharedPtr<FReplicatedObjectData>* Item = PathToObjectDataCache.Find(ChildObject))
 					{
@@ -274,21 +274,19 @@ namespace UE::ConcertSharedSlate
 
 	void SReplicationStreamViewer::BuildObjectHierarchyIfNeeded(TSharedPtr<FReplicatedObjectData> ReplicatedObjectData, TMap<FSoftObjectPath, TSharedPtr<FReplicatedObjectData>>& NewPathToObjectDataCache)
 	{
-		// We're are not supposed to display any hierarchy in the outliner if SubobjectModel is not set.
+		// We're are not supposed to display any hierarchy in the outliner if ObjectHierarchy is not set.
 		const FSoftObjectPath& ObjectPath = ReplicatedObjectData->GetObjectPath();
-		if (!SubobjectModel)
+		if (!ObjectHierarchy)
 		{
 			return;
 		}
 
 		// Find top level object of ReplicatedObjectData
-		const TOptional<FSoftObjectPath> OwningActor = ObjectUtils::GetActorOf(ObjectPath);
-		if (!SubobjectModel->IsTopLevelObject(ObjectPath) && !OwningActor)
+		const FSoftObjectPath OwningActor = ObjectUtils::GetActorOf(ObjectPath).Get(ObjectPath);
+		if (!ObjectUtils::IsActor(OwningActor))
 		{
 			return;
 		}
-		const FSoftObjectPath TopLevelObject = OwningActor.Get(ObjectPath);
-		SubobjectModel->SetTopLevelObject(TopLevelObject);
 		
 		// Add all objects that appear in the hierarchy of ReplicatedObjectData
 		const auto AddItem = [this, &NewPathToObjectDataCache](const FSoftObjectPath& ObjectPath)
@@ -299,8 +297,8 @@ namespace UE::ConcertSharedSlate
 			AllObjectRowData.AddUnique(Item);
 			NewPathToObjectDataCache.Emplace(ObjectPath, Item);
 		};
-		AddItem(TopLevelObject);
-		SubobjectModel->ForEachSubobject([this, &AddItem](const FSoftObjectPath&, const FSoftObjectPath& ChildObject)
+		AddItem(OwningActor);
+		ObjectHierarchy->ForEachChildRecursive(OwningActor, [this, &AddItem](const FSoftObjectPath&, const FSoftObjectPath& ChildObject, EChildRelationship)
 		{
 			AddItem(ChildObject);
 			return EBreakBehavior::Continue;
@@ -312,43 +310,19 @@ namespace UE::ConcertSharedSlate
 		// Important: this view should be possible to be built in programs, so it should not reference things like AActor, UActorComponent, ResolveObject, etc. directly.
 		
 		const FSoftObjectPath& SearchedObject = ReplicatedObjectData->GetObjectPath();
-		if (!SubobjectModel)
+		if (!ObjectHierarchy)
 		{
 			return;
 		}
 
-		const TOptional<FSoftObjectPath> OwningActor = ObjectUtils::GetActorOf(SearchedObject);
-		if (!SubobjectModel->IsTopLevelObject(SearchedObject) && !OwningActor)
+		ObjectHierarchy->ForEachDirectChild(SearchedObject, [this, &ProcessChild](const FSoftObjectPath& ChildObject, EChildRelationship)
 		{
-			return;
-		}
-		SubobjectModel->SetTopLevelObject(OwningActor.Get(SearchedObject));
-		
-		if (SubobjectModel->IsTopLevelObject(SearchedObject))
-		{
-			for (FName Category : SubobjectModel->GetCategories())
+			if (const TSharedPtr<FReplicatedObjectData>* ObjectData = PathToObjectDataCache.Find(ChildObject))
 			{
-				SubobjectModel->ForEachRootSubobject(Category, [this, &ProcessChild](const FSoftObjectPath& ChildObject)
-				{
-					if (const TSharedPtr<FReplicatedObjectData>* ObjectData = PathToObjectDataCache.Find(ChildObject))
-					{
-						ProcessChild(*ObjectData);
-					}
-					return EBreakBehavior::Continue;
-				});
+				ProcessChild(*ObjectData);
 			}
-		}
-		else
-		{
-			SubobjectModel->ForEachDirectChildSubobject(SearchedObject, [this, &ProcessChild](const FSoftObjectPath& ChildObject)
-			{
-				if (const TSharedPtr<FReplicatedObjectData>* ObjectData = PathToObjectDataCache.Find(ChildObject))
-				{
-					ProcessChild(*ObjectData);
-				}
-				return EBreakBehavior::Continue;
-			});
-		}
+			return EBreakBehavior::Continue;
+		});
 	}
 }
 
diff --git a/Engine/Plugins/Developer/Concert/ConcertUI/ConcertSharedSlate/Source/ConcertSharedSlate/Private/Replication/Editor/View/ObjectViewer/SReplicationStreamViewer.h b/Engine/Plugins/Developer/Concert/ConcertUI/ConcertSharedSlate/Source/ConcertSharedSlate/Private/Replication/Editor/View/ObjectViewer/SReplicationStreamViewer.h
--- a/Engine/Plugins/Developer/Concert/ConcertUI/ConcertSharedSlate/Source/ConcertSharedSlate/Private/Replication/Editor/View/ObjectViewer/SReplicationStreamViewer.h
+++ b/Engine/Plugins/Developer/Concert/ConcertUI/ConcertSharedSlate/Source/ConcertSharedSlate/Private/Replication/Editor/View/ObjectViewer/SReplicationStreamViewer.h
@@ -27,7 +27,7 @@ namespace UE::ConcertSharedSlate
 	class IEditableReplicationStreamModel;
 	class IObjectNameModel;
 	class IReplicationStreamModel;
-	class ISubobjectModel;
+	class IObjectHierarchyModel;
 	class SPropertyTreeView;
 	
 	/**
@@ -61,7 +61,7 @@ namespace UE::ConcertSharedSlate
 			SLATE_ARGUMENT(FColumnSortInfo, SecondaryPropertySort)
 
 			/** Optional. If set, this determines the children nested under the root objects. */
-			SLATE_ARGUMENT(TSharedPtr<ISubobjectModel>, SubobjectModel)
+			SLATE_ARGUMENT(TSharedPtr<IObjectHierarchyModel>, ObjectHierarchy)
 			/** Optional. If set, this determines the display text for objects. */
 			SLATE_ARGUMENT(TSharedPtr<IObjectNameModel>, NameModel)
 
@@ -105,7 +105,7 @@ namespace UE::ConcertSharedSlate
 		/** The model this view is visualizing. */
 		TSharedPtr<IReplicationStreamModel> PropertiesModel;
 		/** Can be null. If set, this determines the children nested under the root objects. Editor builds have access to e.g. to USubobjectDataSubsystem but programs do not. */
-		TSharedPtr<ISubobjectModel> SubobjectModel;
+		TSharedPtr<IObjectHierarchyModel> ObjectHierarchy;
 		/** Can be null. If set, this determines the display text for objects. Editor builds have access to e.g. to USubobjectDataSubsystem but programs do not. */
 		TSharedPtr<IObjectNameModel> NameModel;
 
diff --git a/Engine/Plugins/Developer/Concert/ConcertUI/ConcertSharedSlate/Source/ConcertSharedSlate/Private/Replication/ReplicationWidgetFactories.cpp b/Engine/Plugins/Developer/Concert/ConcertUI/ConcertSharedSlate/Source/ConcertSharedSlate/Private/Replication/ReplicationWidgetFactories.cpp
--- a/Engine/Plugins/Developer/Concert/ConcertUI/ConcertSharedSlate/Source/ConcertSharedSlate/Private/Replication/ReplicationWidgetFactories.cpp
+++ b/Engine/Plugins/Developer/Concert/ConcertUI/ConcertSharedSlate/Source/ConcertSharedSlate/Private/Replication/ReplicationWidgetFactories.cpp
@@ -25,7 +25,7 @@ namespace UE::ConcertSharedSlate
 			.AdditionalPropertyColumns(Params.ViewerParams.AdditionalPropertyColumns)
 			.PrimaryPropertySort(Params.ViewerParams.PrimaryPropertySort)
 			.SecondaryPropertySort(Params.ViewerParams.SecondaryPropertySort)
-			.SubobjectModel(Params.ViewerParams.SubobjectModel)
+			.ObjectHierarchy(Params.ViewerParams.ObjectHierarchy)
 			.NameModel(Params.ViewerParams.NameModel)
 			.OnExtendObjectsContextMenu(Params.ViewerParams.OnExtendObjectsContextMenu)
 			.LeftOfObjectSearchBar() [ Params.ViewerParams.LeftOfObjectSearchBar.Widget ]
diff --git a/Engine/Plugins/Developer/Concert/ConcertUI/ConcertSharedSlate/Source/ConcertSharedSlate/Public/Replication/Editor/Model/ISubobjectModel.h b/Engine/Plugins/Developer/Concert/ConcertUI/ConcertSharedSlate/Source/ConcertSharedSlate/Public/Replication/Editor/Model/ISubobjectModel.h
deleted file mode 100644
diff --git a/Engine/Plugins/Developer/Concert/ConcertUI/ConcertSharedSlate/Source/ConcertSharedSlate/Public/Replication/Editor/Model/Object/IObjectHierarchyModel.h b/Engine/Plugins/Developer/Concert/ConcertUI/ConcertSharedSlate/Source/ConcertSharedSlate/Public/Replication/Editor/Model/Object/IObjectHierarchyModel.h
new file mode 100644
--- /dev/null
+++ b/Engine/Plugins/Developer/Concert/ConcertUI/ConcertSharedSlate/Source/ConcertSharedSlate/Public/Replication/Editor/Model/Object/IObjectHierarchyModel.h
@@ -0,0 +1,84 @@
+// Copyright Epic Games, Inc. All Rights Reserved.
+
+#pragma once
+
+#include "Delegates/Delegate.h"
+#include "Misc/EBreakBehavior.h"
+#include "Templates/Function.h"
+
+enum class EBreakBehavior : uint8;
+struct FSoftObjectPath;
+
+namespace UE::ConcertSharedSlate
+{
+	/** Describes the kind of relationship a child object can have to its parent object */
+	enum class EChildRelationship : uint8
+	{
+		/** The child is a child component */
+		Component,
+		/** The child is subobject but not a component */
+		Subobject,
+
+		/** The relationship is not further specified / unknown (for example in non-editor builds you may need to use this due to missing runtime info) */
+		Unknown
+
+		// Could add another one for child actor
+	};
+
+	/** Similar to EChildRelationship. Used for filtering using bit flags. */
+	enum class EChildRelationshipFlags : uint8
+	{
+		None = 0,
+		
+		/** The child is a component of the parent. */
+		Component = 1 << 0,
+		/** The child is a subobject of the parent but no component. */
+		Subobject = 1 << 1,
+
+		/** The relationship is not further specified / unknown (for example in non-editor builds you may need to use this due to missing runtime info) */
+		Unknown = 1 << 2,
+
+		All = Component | Subobject | Unknown
+	};
+	ENUM_CLASS_FLAGS(EChildRelationshipFlags);
+	
+	/**
+	 * Determines the parent-child relationship between objects.
+	 * The model can be given an object and it will list all child objects, usually subobjects or child components.
+	 *
+	 * In editor builds, an implementation could resolve object paths to AActors, USceneComponent, etc.
+	 * In programs, an implementation could be to compare all known object path strings.
+	 */
+	class CONCERTSHAREDSLATE_API IObjectHierarchyModel
+	{
+	public:
+		
+		/**
+		 * Gets all direct children in the hierarchy.
+		 *
+		 * An example for what ForEachDirectChild could return in an editor build:
+		 * - Actor returns its root component, non-scene components (UActorComponents), and all other direct subobjects
+		 * - Scene components return their children, and all other direct subobjects
+		 * - All other subobjects report all direct subobjects
+		 * The hierarchy is implementation defined: the above is just an example for describing how ForEachDirectChild should operate.
+		 *
+		 * @param Root
+		 * @param Callback
+		 * @param InclusionFlags Specifies the relationships that should be included
+		 */
+		virtual void ForEachDirectChild(
+			const FSoftObjectPath& Root,
+			TFunctionRef<EBreakBehavior(const FSoftObjectPath& Object, EChildRelationship Relationship)> Callback,
+			EChildRelationshipFlags InclusionFlags = EChildRelationshipFlags::All
+			) const = 0;
+
+		/** Util for iterating all subobjects. */
+		void ForEachChildRecursive(
+			const FSoftObjectPath& Root,
+			TFunctionRef<EBreakBehavior(const FSoftObjectPath& Parent, const FSoftObjectPath& ChildObject, EChildRelationship Relationship)> Callback,
+			EChildRelationshipFlags InclusionFlags = EChildRelationshipFlags::All
+			);
+
+		virtual ~IObjectHierarchyModel() = default;
+	};
+}
diff --git a/Engine/Plugins/Developer/Concert/ConcertUI/ConcertSharedSlate/Source/ConcertSharedSlate/Public/Replication/ReplicationWidgetFactories.h b/Engine/Plugins/Developer/Concert/ConcertUI/ConcertSharedSlate/Source/ConcertSharedSlate/Public/Replication/ReplicationWidgetFactories.h
--- a/Engine/Plugins/Developer/Concert/ConcertUI/ConcertSharedSlate/Source/ConcertSharedSlate/Public/Replication/ReplicationWidgetFactories.h
+++ b/Engine/Plugins/Developer/Concert/ConcertUI/ConcertSharedSlate/Source/ConcertSharedSlate/Public/Replication/ReplicationWidgetFactories.h
@@ -29,7 +29,7 @@ namespace UE::ConcertSharedSlate
 	class IReplicationStreamViewer;
 	class IReplicationSubobjectView;
 	class IStreamExtender;
-	class ISubobjectModel;
+	class IObjectHierarchyModel;
 	class IPropertySelectionSourceModel;
 	
 	/**
@@ -53,7 +53,7 @@ namespace UE::ConcertSharedSlate
 		 * If left unspecified, the top will only display actors.
 		 * @note The created view will keep a strong reference to this.
 		 */
-		TSharedPtr<ISubobjectModel> SubobjectModel;
+		TSharedPtr<IObjectHierarchyModel> ObjectHierarchy;
 
 		/** Optional. Determines the display name of objects. */
 		TSharedPtr<IObjectNameModel> NameModel;
