Add view options next to object search box with which subobjects can be hidden (filtered out by default)

diff --git a/Engine/Plugins/Developer/Concert/ConcertUI/ConcertClientSharedSlate/Source/ConcertClientSharedSlate/Private/Replication/Editor/Model/Object/EditorObjectHierarchyModel.cpp b/Engine/Plugins/Developer/Concert/ConcertUI/ConcertClientSharedSlate/Source/ConcertClientSharedSlate/Private/Replication/Editor/Model/Object/EditorObjectHierarchyModel.cpp
--- a/Engine/Plugins/Developer/Concert/ConcertUI/ConcertClientSharedSlate/Source/ConcertClientSharedSlate/Private/Replication/Editor/Model/Object/EditorObjectHierarchyModel.cpp
+++ b/Engine/Plugins/Developer/Concert/ConcertUI/ConcertClientSharedSlate/Source/ConcertClientSharedSlate/Private/Replication/Editor/Model/Object/EditorObjectHierarchyModel.cpp
@@ -92,4 +92,27 @@ namespace UE::ConcertClientSharedSlate
 			}, bIncludeNestedSubobjects);
 		}
 	}
-}
\ No newline at end of file
+
+	TOptional<ConcertSharedSlate::IObjectHierarchyModel::FParentInfo> FEditorObjectHierarchyModel::GetParentInfo(const FSoftObjectPath& ChildObject) const
+	{
+		using namespace ConcertSharedSlate;
+
+		// In the context of replication, only resolve objects but do not load them. Case: client 1 may be viewing objects of client 2 but they are in different worlds.
+		const UObject* Object = ChildObject.ResolveObject();
+		if (!Object)
+		{
+			return {};
+		}
+
+		if (Object->IsA<AActor>()
+			// Unlikely but theoretically possible that it's an asset or some transient object.
+			|| !Object->IsInA(AActor::StaticClass()))
+		{
+			return {};
+		}
+		
+		const UObject* Outer = Object->GetOuter();
+		const EChildRelationship Relationship = Object->IsA<UActorComponent>() ? EChildRelationship::Component : EChildRelationship::Subobject;
+		return FParentInfo{ Outer, Relationship };
+	}
+}
diff --git a/Engine/Plugins/Developer/Concert/ConcertUI/ConcertClientSharedSlate/Source/ConcertClientSharedSlate/Private/Replication/Editor/Model/Object/EditorObjectHierarchyModel.h b/Engine/Plugins/Developer/Concert/ConcertUI/ConcertClientSharedSlate/Source/ConcertClientSharedSlate/Private/Replication/Editor/Model/Object/EditorObjectHierarchyModel.h
--- a/Engine/Plugins/Developer/Concert/ConcertUI/ConcertClientSharedSlate/Source/ConcertClientSharedSlate/Private/Replication/Editor/Model/Object/EditorObjectHierarchyModel.h
+++ b/Engine/Plugins/Developer/Concert/ConcertUI/ConcertClientSharedSlate/Source/ConcertClientSharedSlate/Private/Replication/Editor/Model/Object/EditorObjectHierarchyModel.h
@@ -27,6 +27,7 @@ namespace UE::ConcertClientSharedSlate
 			TFunctionRef<EBreakBehavior(const FSoftObjectPath& Object, ConcertSharedSlate::EChildRelationship Relationship)> Callback,
 			ConcertSharedSlate::EChildRelationshipFlags InclusionFlags
 			) const override;
+		virtual TOptional<FParentInfo> GetParentInfo(const FSoftObjectPath& ChildObject) const override;
 		//~ End IObjectHierarchy Interface
 	};
 }
diff --git a/Engine/Plugins/Developer/Concert/ConcertUI/ConcertSharedSlate/Source/ConcertSharedSlate/Private/Replication/Editor/View/ObjectEditor/SBaseReplicationStreamEditor.h b/Engine/Plugins/Developer/Concert/ConcertUI/ConcertSharedSlate/Source/ConcertSharedSlate/Private/Replication/Editor/View/ObjectEditor/SBaseReplicationStreamEditor.h
--- a/Engine/Plugins/Developer/Concert/ConcertUI/ConcertSharedSlate/Source/ConcertSharedSlate/Private/Replication/Editor/View/ObjectEditor/SBaseReplicationStreamEditor.h
+++ b/Engine/Plugins/Developer/Concert/ConcertUI/ConcertSharedSlate/Source/ConcertSharedSlate/Private/Replication/Editor/View/ObjectEditor/SBaseReplicationStreamEditor.h
@@ -76,8 +76,13 @@ namespace UE::ConcertSharedSlate
 
 			/** Optional widget to add to the left of the object list search bar. */
 			SLATE_NAMED_SLOT(FArguments, LeftOfObjectSearchBar)
+			/** Optional widget to add to the right of the object list search bar. */
+			SLATE_NAMED_SLOT(FArguments, RightOfObjectSearchBar)
+		
 			/** Optional widget to add to the left of the property list search bar. */
 			SLATE_NAMED_SLOT(FArguments, LeftOfPropertySearchBar)
+			/** Optional widget to add to the right of the property list search bar. */
+			SLATE_NAMED_SLOT(FArguments, RightOfPropertySearchBar)
 
 			/** Optional. Determines whether all UI for changing the model should be disabled. */
 			SLATE_ATTRIBUTE(bool, IsEditingEnabled)
diff --git a/Engine/Plugins/Developer/Concert/ConcertUI/ConcertSharedSlate/Source/ConcertSharedSlate/Private/Replication/Editor/View/ObjectViewer/Property/SPropertyTreeView.cpp b/Engine/Plugins/Developer/Concert/ConcertUI/ConcertSharedSlate/Source/ConcertSharedSlate/Private/Replication/Editor/View/ObjectViewer/Property/SPropertyTreeView.cpp
--- a/Engine/Plugins/Developer/Concert/ConcertUI/ConcertSharedSlate/Source/ConcertSharedSlate/Private/Replication/Editor/View/ObjectViewer/Property/SPropertyTreeView.cpp
+++ b/Engine/Plugins/Developer/Concert/ConcertUI/ConcertSharedSlate/Source/ConcertSharedSlate/Private/Replication/Editor/View/ObjectViewer/Property/SPropertyTreeView.cpp
@@ -86,6 +86,7 @@ namespace UE::ConcertSharedSlate
 						InArgs._LeftOfSearchBar.Widget
 					]
 				]
+				.RightOfSearchBar() [ InArgs._RightOfSearchBar.Widget ]
 				.RowBelowSearchBar() [ FilterBar.ToSharedRef() ]
 				.NoItemsContent() [ SNew(STextBlock).Text(this, &SPropertyTreeView::GetAllFilteredText) ]
 		];
diff --git a/Engine/Plugins/Developer/Concert/ConcertUI/ConcertSharedSlate/Source/ConcertSharedSlate/Private/Replication/Editor/View/ObjectViewer/Property/SPropertyTreeView.h b/Engine/Plugins/Developer/Concert/ConcertUI/ConcertSharedSlate/Source/ConcertSharedSlate/Private/Replication/Editor/View/ObjectViewer/Property/SPropertyTreeView.h
--- a/Engine/Plugins/Developer/Concert/ConcertUI/ConcertSharedSlate/Source/ConcertSharedSlate/Private/Replication/Editor/View/ObjectViewer/Property/SPropertyTreeView.h
+++ b/Engine/Plugins/Developer/Concert/ConcertUI/ConcertSharedSlate/Source/ConcertSharedSlate/Private/Replication/Editor/View/ObjectViewer/Property/SPropertyTreeView.h
@@ -48,6 +48,8 @@ namespace UE::ConcertSharedSlate
 		
 			/** Optional widget to add to the left of the search bar. */
 			SLATE_NAMED_SLOT(FArguments, LeftOfSearchBar)
+			/** Optional widget to add to the left of the search bar. */
+			SLATE_NAMED_SLOT(FArguments, RightOfSearchBar)
 
 			/*************** Own arguments ***************/
 			// Please add new arguments here in the future
diff --git a/Engine/Plugins/Developer/Concert/ConcertUI/ConcertSharedSlate/Source/ConcertSharedSlate/Private/Replication/Editor/View/ObjectViewer/SReplicatedPropertyView.cpp b/Engine/Plugins/Developer/Concert/ConcertUI/ConcertSharedSlate/Source/ConcertSharedSlate/Private/Replication/Editor/View/ObjectViewer/SReplicatedPropertyView.cpp
--- a/Engine/Plugins/Developer/Concert/ConcertUI/ConcertSharedSlate/Source/ConcertSharedSlate/Private/Replication/Editor/View/ObjectViewer/SReplicatedPropertyView.cpp
+++ b/Engine/Plugins/Developer/Concert/ConcertUI/ConcertSharedSlate/Source/ConcertSharedSlate/Private/Replication/Editor/View/ObjectViewer/SReplicatedPropertyView.cpp
@@ -130,10 +130,8 @@ namespace UE::ConcertSharedSlate
 				.SecondarySort(SecondarySort)
 				.SelectionMode(ESelectionMode::Multi)
 				.NameModel(InArgs._NameModel)
-				.LeftOfSearchBar()
-				[
-					InArgs._LeftOfPropertySearchBar.Widget
-				]
+				.LeftOfSearchBar() [ InArgs._LeftOfPropertySearchBar.Widget ]
+				.RightOfSearchBar() [ InArgs._RightOfPropertySearchBar.Widget ]
 				.SelectedObjects(this, &SReplicatedPropertyView::GetObjectsSelectedForPropertyEditing)
 			]
 			
diff --git a/Engine/Plugins/Developer/Concert/ConcertUI/ConcertSharedSlate/Source/ConcertSharedSlate/Private/Replication/Editor/View/ObjectViewer/SReplicatedPropertyView.h b/Engine/Plugins/Developer/Concert/ConcertUI/ConcertSharedSlate/Source/ConcertSharedSlate/Private/Replication/Editor/View/ObjectViewer/SReplicatedPropertyView.h
--- a/Engine/Plugins/Developer/Concert/ConcertUI/ConcertSharedSlate/Source/ConcertSharedSlate/Private/Replication/Editor/View/ObjectViewer/SReplicatedPropertyView.h
+++ b/Engine/Plugins/Developer/Concert/ConcertUI/ConcertSharedSlate/Source/ConcertSharedSlate/Private/Replication/Editor/View/ObjectViewer/SReplicatedPropertyView.h
@@ -39,6 +39,8 @@ namespace UE::ConcertSharedSlate
 		
 			/** Optional widget to add to the left of the property list search bar. */
 			SLATE_NAMED_SLOT(FArguments, LeftOfPropertySearchBar)
+			/** Optional widget to add to the right of the property list search bar. */
+			SLATE_NAMED_SLOT(FArguments, RightOfPropertySearchBar)
 		SLATE_END_ARGS()
 
 		void Construct(const FArguments& InArgs, TSharedRef<IReplicationStreamModel> InPropertiesModel);
diff --git a/Engine/Plugins/Developer/Concert/ConcertUI/ConcertSharedSlate/Source/ConcertSharedSlate/Private/Replication/Editor/View/ObjectViewer/SReplicationStreamViewer.cpp b/Engine/Plugins/Developer/Concert/ConcertUI/ConcertSharedSlate/Source/ConcertSharedSlate/Private/Replication/Editor/View/ObjectViewer/SReplicationStreamViewer.cpp
--- a/Engine/Plugins/Developer/Concert/ConcertUI/ConcertSharedSlate/Source/ConcertSharedSlate/Private/Replication/Editor/View/ObjectViewer/SReplicationStreamViewer.cpp
+++ b/Engine/Plugins/Developer/Concert/ConcertUI/ConcertSharedSlate/Source/ConcertSharedSlate/Private/Replication/Editor/View/ObjectViewer/SReplicationStreamViewer.cpp
@@ -68,15 +68,21 @@ namespace UE::ConcertSharedSlate
 		// Do a complete refresh.
 		// Complete refresh is acceptable because the list is updated infrequently and typically small < 500 items.
 		// An alternative would be to change RefreshObjectData to be called with two variables ObjectsAdded and ObjectsRemoved.
-		PropertiesModel->ForEachReplicatedObject([this, &NewPathToObjectDataCache](const FSoftObjectPath& Path) mutable
+		PropertiesModel->ForEachReplicatedObject([this, &NewPathToObjectDataCache](const FSoftObjectPath& ObjectPath) mutable
 		{
-			const TSharedPtr<FReplicatedObjectData>* ExistingItem = PathToObjectDataCache.Find(Path);
-			ExistingItem = ExistingItem ? ExistingItem : NewPathToObjectDataCache.Find(Path);
-			const TSharedRef<FReplicatedObjectData> Item = ExistingItem ? ExistingItem->ToSharedRef() : AllocateObjectData(Path);
-			AllObjectRowData.AddUnique(Item);
-			NewPathToObjectDataCache.Emplace(Path, Item);
+			TOptional<IObjectHierarchyModel::FParentInfo> ParentInfo = ObjectHierarchy->GetParentInfo(ObjectPath);
+			const bool bIsActor = !ParentInfo; 
+			if (bIsActor || ShouldDisplayObject(ObjectPath, ParentInfo->Relationship))
+			{
+				const TSharedPtr<FReplicatedObjectData>* ExistingItem = PathToObjectDataCache.Find(ObjectPath);
+				ExistingItem = ExistingItem ? ExistingItem : NewPathToObjectDataCache.Find(ObjectPath);
+				const TSharedRef<FReplicatedObjectData> Item = ExistingItem ? ExistingItem->ToSharedRef() : AllocateObjectData(ObjectPath);
+				AllObjectRowData.AddUnique(Item);
+				NewPathToObjectDataCache.Emplace(ObjectPath, Item);
+				
+				BuildObjectHierarchyIfNeeded(Item, NewPathToObjectDataCache);
+			}
 			
-			BuildObjectHierarchyIfNeeded(Item, NewPathToObjectDataCache);
 			return EBreakBehavior::Continue;
 		});
 
@@ -190,6 +196,8 @@ namespace UE::ConcertSharedSlate
 		const FColumnSortInfo SecondaryObjectSort = InArgs._SecondaryObjectSort.IsValid()
 			? InArgs._SecondaryObjectSort
 			: FColumnSortInfo{ ReplicationColumns::TopLevel::LabelColumnId, EColumnSortMode::Ascending };
+
+		ObjectViewOptions.OnDisplaySubobjectsToggled().AddSP(this, &SReplicationStreamViewer::OnSubobjectViewOptionToggled);
 		
 		return SAssignNew(ReplicatedObjects, SReplicationTreeView<FReplicatedObjectData>)
 			.RootItemsSource(&RootObjectRowData)
@@ -205,6 +213,22 @@ namespace UE::ConcertSharedSlate
 			.SecondarySort(SecondaryObjectSort)
 			.SelectionMode(ESelectionMode::Multi)
 			.LeftOfSearchBar() [ InArgs._LeftOfObjectSearchBar.Widget ]
+			.RightOfSearchBar()
+			[
+				SNew(SHorizontalBox)
+
+				+SHorizontalBox::Slot()
+				.AutoWidth()
+				[
+					InArgs._RightOfObjectSearchBar.Widget
+				]
+				
+				+SHorizontalBox::Slot()
+				.AutoWidth()
+				[
+					ObjectViewOptions.MakeViewOptionsComboButton()
+				]
+			]
 			.NoItemsContent() [ SNew(STextBlock).Text(NoObjectsAttribute) ];
 	}
 
@@ -237,10 +261,8 @@ namespace UE::ConcertSharedSlate
 					.SecondarySort(InArgs._SecondaryPropertySort)
 					.GetSelectedRootObjects_Lambda([this](){ return GetSelectedOutlinerObjects(); })
 					.NameModel(InArgs._NameModel)
-					.LeftOfPropertySearchBar()
-					[
-						InArgs._LeftOfPropertySearchBar.Widget
-					]
+					.LeftOfPropertySearchBar() [ InArgs._LeftOfPropertySearchBar.Widget ]
+					.RightOfPropertySearchBar() [ InArgs._RightOfPropertySearchBar.Widget ]
 				]
 			];
 	}
@@ -297,10 +319,14 @@ namespace UE::ConcertSharedSlate
 			AllObjectRowData.AddUnique(Item);
 			NewPathToObjectDataCache.Emplace(ObjectPath, Item);
 		};
+		
 		AddItem(OwningActor);
-		ObjectHierarchy->ForEachChildRecursive(OwningActor, [this, &AddItem](const FSoftObjectPath&, const FSoftObjectPath& ChildObject, EChildRelationship)
+		ObjectHierarchy->ForEachChildRecursive(OwningActor, [this, &AddItem](const FSoftObjectPath&, const FSoftObjectPath& ChildObject, EChildRelationship Relationship)
 		{
-			AddItem(ChildObject);
+			if (ShouldDisplayObject(ChildObject, Relationship))
+			{
+				AddItem(ChildObject);
+			}
 			return EBreakBehavior::Continue;
 		});
 	}
@@ -315,7 +341,7 @@ namespace UE::ConcertSharedSlate
 			return;
 		}
 
-		ObjectHierarchy->ForEachDirectChild(SearchedObject, [this, &ProcessChild](const FSoftObjectPath& ChildObject, EChildRelationship)
+		ObjectHierarchy->ForEachDirectChild(SearchedObject, [this, &ProcessChild](const FSoftObjectPath& ChildObject, EChildRelationship Relationship)
 		{
 			if (const TSharedPtr<FReplicatedObjectData>* ObjectData = PathToObjectDataCache.Find(ChildObject))
 			{
@@ -324,6 +350,12 @@ namespace UE::ConcertSharedSlate
 			return EBreakBehavior::Continue;
 		});
 	}
+
+	bool SReplicationStreamViewer::ShouldDisplayObject(const FSoftObjectPath& Object, EChildRelationship Relationship) const
+	{
+		const bool bSkipSubobject = Relationship == EChildRelationship::Subobject && !ObjectViewOptions.ShouldDisplaySubobjects();
+		return !bSkipSubobject;
+	}
 }
 
 #undef LOCTEXT_NAMESPACE
\ No newline at end of file
diff --git a/Engine/Plugins/Developer/Concert/ConcertUI/ConcertSharedSlate/Source/ConcertSharedSlate/Private/Replication/Editor/View/ObjectViewer/SReplicationStreamViewer.h b/Engine/Plugins/Developer/Concert/ConcertUI/ConcertSharedSlate/Source/ConcertSharedSlate/Private/Replication/Editor/View/ObjectViewer/SReplicationStreamViewer.h
--- a/Engine/Plugins/Developer/Concert/ConcertUI/ConcertSharedSlate/Source/ConcertSharedSlate/Private/Replication/Editor/View/ObjectViewer/SReplicationStreamViewer.h
+++ b/Engine/Plugins/Developer/Concert/ConcertUI/ConcertSharedSlate/Source/ConcertSharedSlate/Private/Replication/Editor/View/ObjectViewer/SReplicationStreamViewer.h
@@ -8,6 +8,7 @@
 #include "Replication/Editor/View/Tree/SReplicationTreeView.h"
 #include "Replication/Editor/View/SelectionViewerColumns.h"
 #include "SReplicatedPropertyView.h"
+#include "StreamViewerObjectViewOptions.h"
 
 #include "Algo/Transform.h"
 #include "Misc/Optional.h"
@@ -30,6 +31,8 @@ namespace UE::ConcertSharedSlate
 	class IObjectHierarchyModel;
 	class SPropertyTreeView;
 	
+	enum class EChildRelationship : uint8;
+	
 	/**
 	 * Root widget for viewing UMultiUserPropertyReplicationSelection.
 	 * This widget knows how to display IReplicationStreamModel.
@@ -73,8 +76,13 @@ namespace UE::ConcertSharedSlate
 		
 			/** Optional widget to add to the left of the object list search bar. */
 			SLATE_NAMED_SLOT(FArguments, LeftOfObjectSearchBar)
+			/** Optional widget to add to the right of the object list search bar. */
+			SLATE_NAMED_SLOT(FArguments, RightOfObjectSearchBar)
+		
 			/** Optional widget to add to the left of the property list search bar. */
 			SLATE_NAMED_SLOT(FArguments, LeftOfPropertySearchBar)
+			/** Optional widget to add to the right of the object list search bar. */
+			SLATE_NAMED_SLOT(FArguments, RightOfPropertySearchBar)
 
 			/** Optional text to display when no object is in the outliner. Defaults to "No objects to display." "*/
 			SLATE_ATTRIBUTE(FText, NoOutlinerObjects)
@@ -125,6 +133,8 @@ namespace UE::ConcertSharedSlate
 		TMap<FSoftObjectPath, TSharedPtr<FReplicatedObjectData>> PathToObjectDataCache;
 
 		bool bIsPropertyAreaExpanded = false;
+		/** View options for the object outliner. */
+		FStreamViewerObjectViewOptions ObjectViewOptions;
 
 		static TSharedRef<FReplicatedObjectData> AllocateObjectData(FSoftObjectPath ObjectPath);
 
@@ -143,5 +153,13 @@ namespace UE::ConcertSharedSlate
 		/** Handles how much space the 'Clients' area uses with respect to its expansion state. */
 		SSplitter::ESizeRule GetPropertyAreaSizeRule() const { return bIsPropertyAreaExpanded ? SSplitter::ESizeRule::FractionOfParent : SSplitter::ESizeRule::SizeToContent; }
 		void OnPropertyAreaExpansionChanged(bool bExpanded) { bIsPropertyAreaExpanded = bExpanded; }
+
+		/** Called in response to subobject display view option being changed. Rebuilds the entire hierarchy. */
+		void OnSubobjectViewOptionToggled()
+		{
+			RefreshObjectData();
+			RefreshPropertyData();
+		}
+		bool ShouldDisplayObject(const FSoftObjectPath& Object, EChildRelationship Relationship) const;
 	};
 }
diff --git a/Engine/Plugins/Developer/Concert/ConcertUI/ConcertSharedSlate/Source/ConcertSharedSlate/Private/Replication/Editor/View/ObjectViewer/StreamViewerObjectViewOptions.cpp b/Engine/Plugins/Developer/Concert/ConcertUI/ConcertSharedSlate/Source/ConcertSharedSlate/Private/Replication/Editor/View/ObjectViewer/StreamViewerObjectViewOptions.cpp
new file mode 100644
--- /dev/null
+++ b/Engine/Plugins/Developer/Concert/ConcertUI/ConcertSharedSlate/Source/ConcertSharedSlate/Private/Replication/Editor/View/ObjectViewer/StreamViewerObjectViewOptions.cpp
@@ -0,0 +1,40 @@
+// Copyright Epic Games, Inc. All Rights Reserved.
+
+#include "StreamViewerObjectViewOptions.h"
+
+#include "ConcertFrontendUtils.h"
+
+#include "Framework/MultiBox/MultiBoxBuilder.h"
+
+#define LOCTEXT_NAMESPACE "FStreamViewerObjectViewOptions"
+
+namespace UE::ConcertSharedSlate
+{
+	TSharedRef<SWidget> FStreamViewerObjectViewOptions::MakeViewOptionsComboButton()
+	{
+		return ConcertFrontendUtils::CreateViewOptionsComboButton(
+			FOnGetContent::CreateRaw(this, &FStreamViewerObjectViewOptions::MakeMenuWidget)
+			);
+	}
+
+	TSharedRef<SWidget> FStreamViewerObjectViewOptions::MakeMenuWidget()
+	{
+		FMenuBuilder MenuBuilder(/*bInShouldCloseWindowAfterMenuSelection=*/true, nullptr);
+
+		MenuBuilder.AddMenuEntry(
+			LOCTEXT("DisplaySubobjects.Label", "Display Subobjects"),
+			LOCTEXT("DisplaySubobjects.Tooltip", "Whether to display subobjects (with the exception of components, which are continued to be displayed).\nSubobjects are objects nested under other objects."),
+			FSlateIcon(),
+			FUIAction(
+				FExecuteAction::CreateLambda([this](){ bDisplaySubobjects = !bDisplaySubobjects; OnDisplaySubobjectsToggledDelegate.Broadcast(); }),
+				FCanExecuteAction::CreateLambda([] { return true; }),
+				FIsActionChecked::CreateLambda([this] { return bDisplaySubobjects; })),
+			NAME_None,
+			EUserInterfaceActionType::ToggleButton
+		);
+
+		return MenuBuilder.MakeWidget();
+	}
+}
+
+#undef LOCTEXT_NAMESPACE
\ No newline at end of file
diff --git a/Engine/Plugins/Developer/Concert/ConcertUI/ConcertSharedSlate/Source/ConcertSharedSlate/Private/Replication/Editor/View/ObjectViewer/StreamViewerObjectViewOptions.h b/Engine/Plugins/Developer/Concert/ConcertUI/ConcertSharedSlate/Source/ConcertSharedSlate/Private/Replication/Editor/View/ObjectViewer/StreamViewerObjectViewOptions.h
new file mode 100644
--- /dev/null
+++ b/Engine/Plugins/Developer/Concert/ConcertUI/ConcertSharedSlate/Source/ConcertSharedSlate/Private/Replication/Editor/View/ObjectViewer/StreamViewerObjectViewOptions.h
@@ -0,0 +1,38 @@
+// Copyright Epic Games, Inc. All Rights Reserved.
+
+#pragma once
+
+#include "Delegates/Delegate.h"
+#include "Templates/SharedPointer.h"
+
+class SWidget;
+
+namespace UE::ConcertSharedSlate
+{
+	/** View options for the outliner displayed by SReplicationStreamViewer. */
+	class FStreamViewerObjectViewOptions
+    {
+    public:
+
+		/** Makes a standard View Options widget, displaying the eye ball icon and showing the possible options. */
+		TSharedRef<SWidget> MakeViewOptionsComboButton();
+		/** Returns a Menu widgets containing the available options. */
+		TSharedRef<SWidget> MakeMenuWidget();
+
+		/** Whether to display subobjects (with the exception of components, which are continued to be displayed). */
+		bool ShouldDisplaySubobjects() const { return bDisplaySubobjects; }
+
+		DECLARE_MULTICAST_DELEGATE(FOnDisplaySubobjectsToggled);
+		/** Called when the option for displaying subobjects has been changed. */
+		FOnDisplaySubobjectsToggled& OnDisplaySubobjectsToggled() { return OnDisplaySubobjectsToggledDelegate; }
+		
+	private:
+
+		/** Called when the option for displaying subobjects has been changed. */
+		FOnDisplaySubobjectsToggled OnDisplaySubobjectsToggledDelegate;
+		
+		/** Whether to display subobjects (with the exception of components, which are continued to be displayed). */
+		bool bDisplaySubobjects = false;
+    };
+}
+
diff --git a/Engine/Plugins/Developer/Concert/ConcertUI/ConcertSharedSlate/Source/ConcertSharedSlate/Private/Replication/Editor/View/Tree/SReplicationTreeView.h b/Engine/Plugins/Developer/Concert/ConcertUI/ConcertSharedSlate/Source/ConcertSharedSlate/Private/Replication/Editor/View/Tree/SReplicationTreeView.h
--- a/Engine/Plugins/Developer/Concert/ConcertUI/ConcertSharedSlate/Source/ConcertSharedSlate/Private/Replication/Editor/View/Tree/SReplicationTreeView.h
+++ b/Engine/Plugins/Developer/Concert/ConcertUI/ConcertSharedSlate/Source/ConcertSharedSlate/Private/Replication/Editor/View/Tree/SReplicationTreeView.h
@@ -89,6 +89,8 @@ namespace UE::ConcertSharedSlate
 
 			/** Optional widget to add to the left of the search bar. */
 			SLATE_NAMED_SLOT(FArguments, LeftOfSearchBar)
+			/** Optional widget to add to the right of the search bar. */
+			SLATE_NAMED_SLOT(FArguments, RightOfSearchBar)
 			/** Optional widget to add between the search bar and the table view (e.g. a SBasicFilterBar). */
 			SLATE_NAMED_SLOT(FArguments, RowBelowSearchBar)
 
@@ -140,6 +142,13 @@ namespace UE::ConcertSharedSlate
 						.OnTextCommitted(this, &SReplicationTreeView::OnSearchTextCommitted)
 						.DelayChangeNotificationsWhileTyping(true)
 					]
+
+					+SHorizontalBox::Slot()
+					.Padding(1.f)
+					.AutoWidth()
+					[
+						InArgs._RightOfSearchBar.Widget
+					]
 				]
 
 				// Optional slot between search bar and table, e.g. for an external SBasicFilterBar
diff --git a/Engine/Plugins/Developer/Concert/ConcertUI/ConcertSharedSlate/Source/ConcertSharedSlate/Private/Replication/ReplicationWidgetFactories.cpp b/Engine/Plugins/Developer/Concert/ConcertUI/ConcertSharedSlate/Source/ConcertSharedSlate/Private/Replication/ReplicationWidgetFactories.cpp
--- a/Engine/Plugins/Developer/Concert/ConcertUI/ConcertSharedSlate/Source/ConcertSharedSlate/Private/Replication/ReplicationWidgetFactories.cpp
+++ b/Engine/Plugins/Developer/Concert/ConcertUI/ConcertSharedSlate/Source/ConcertSharedSlate/Private/Replication/ReplicationWidgetFactories.cpp
@@ -29,7 +29,9 @@ namespace UE::ConcertSharedSlate
 			.NameModel(Params.ViewerParams.NameModel)
 			.OnExtendObjectsContextMenu(Params.ViewerParams.OnExtendObjectsContextMenu)
 			.LeftOfObjectSearchBar() [ Params.ViewerParams.LeftOfObjectSearchBar.Widget ]
+			.RightOfObjectSearchBar() [ Params.ViewerParams.RightOfObjectSearchBar.Widget ]
 			.LeftOfPropertySearchBar() [ Params.ViewerParams.LeftOfPropertySearchBar.Widget ]
+			.RightOfPropertySearchBar() [ Params.ViewerParams.RightOfPropertySearchBar.Widget ]
 			.IsEditingEnabled(Params.IsEditingEnabled)
 			.EditingDisabledToolTipText(Params.EditingDisabledToolTipText);
 	}
diff --git a/Engine/Plugins/Developer/Concert/ConcertUI/ConcertSharedSlate/Source/ConcertSharedSlate/Public/Replication/Editor/Model/Object/IObjectHierarchyModel.h b/Engine/Plugins/Developer/Concert/ConcertUI/ConcertSharedSlate/Source/ConcertSharedSlate/Public/Replication/Editor/Model/Object/IObjectHierarchyModel.h
--- a/Engine/Plugins/Developer/Concert/ConcertUI/ConcertSharedSlate/Source/ConcertSharedSlate/Public/Replication/Editor/Model/Object/IObjectHierarchyModel.h
+++ b/Engine/Plugins/Developer/Concert/ConcertUI/ConcertSharedSlate/Source/ConcertSharedSlate/Public/Replication/Editor/Model/Object/IObjectHierarchyModel.h
@@ -4,10 +4,10 @@
 
 #include "Delegates/Delegate.h"
 #include "Misc/EBreakBehavior.h"
+#include "UObject/SoftObjectPath.h"
 #include "Templates/Function.h"
 
 enum class EBreakBehavior : uint8;
-struct FSoftObjectPath;
 
 namespace UE::ConcertSharedSlate
 {
@@ -72,6 +72,16 @@ namespace UE::ConcertSharedSlate
 			EChildRelationshipFlags InclusionFlags = EChildRelationshipFlags::All
 			) const = 0;
 
+		struct FParentInfo
+		{
+			/** The parent of the child object in the hierarchy. */
+			FSoftObjectPath Parent;
+			/** Relationship that the child has to Parent, e.g. the child could be a component of Parent. */
+			EChildRelationship Relationship;
+		};
+		/** Gets parent info for ChildObject, if it has a child. */
+		virtual TOptional<FParentInfo> GetParentInfo(const FSoftObjectPath& ChildObject) const = 0;
+
 		/** Util for iterating all subobjects. */
 		void ForEachChildRecursive(
 			const FSoftObjectPath& Root,
diff --git a/Engine/Plugins/Developer/Concert/ConcertUI/ConcertSharedSlate/Source/ConcertSharedSlate/Public/Replication/ReplicationWidgetFactories.h b/Engine/Plugins/Developer/Concert/ConcertUI/ConcertSharedSlate/Source/ConcertSharedSlate/Public/Replication/ReplicationWidgetFactories.h
--- a/Engine/Plugins/Developer/Concert/ConcertUI/ConcertSharedSlate/Source/ConcertSharedSlate/Public/Replication/ReplicationWidgetFactories.h
+++ b/Engine/Plugins/Developer/Concert/ConcertUI/ConcertSharedSlate/Source/ConcertSharedSlate/Public/Replication/ReplicationWidgetFactories.h
@@ -77,8 +77,12 @@ namespace UE::ConcertSharedSlate
 		
 		/** Optional widget to add to the left of the object list search bar. */
 		TAlwaysValidWidget LeftOfObjectSearchBar;
+		/** Optional widget to add to the right of the object list search bar. */
+		TAlwaysValidWidget RightOfObjectSearchBar;
 		/** Optional widget to add to the left of the property list search bar. */
 		TAlwaysValidWidget LeftOfPropertySearchBar;
+		/** Optional widget to add to the right of the property list search bar. */
+		TAlwaysValidWidget RightOfPropertySearchBar;
 	};
 
 	// TODO DP 5.5: Create factory function that uses FCreateViewerParams and creates an IReplicationStreamViewer
