Make rows in replication view of Multi-User alternate in grey and dark grey

diff --git a/Engine/Plugins/Developer/Concert/ConcertUI/ConcertSharedSlate/Source/ConcertSharedSlate/Private/Replication/Editor/View/Tree/SReplicationColumnRow.h b/Engine/Plugins/Developer/Concert/ConcertUI/ConcertSharedSlate/Source/ConcertSharedSlate/Private/Replication/Editor/View/Tree/SReplicationColumnRow.h
--- a/Engine/Plugins/Developer/Concert/ConcertUI/ConcertSharedSlate/Source/ConcertSharedSlate/Private/Replication/Editor/View/Tree/SReplicationColumnRow.h
+++ b/Engine/Plugins/Developer/Concert/ConcertUI/ConcertSharedSlate/Source/ConcertSharedSlate/Private/Replication/Editor/View/Tree/SReplicationColumnRow.h
@@ -62,7 +62,12 @@ namespace UE::ConcertSharedSlate
 			ExpandableColumnLabel = InArgs._ExpandableColumnLabel;
 			RowHeight = InArgs._RowHeight;
 			
-			SMultiColumnTableRow<TSharedPtr<TListItemType>>::Construct({}, InOwner);
+			using FTableRowArgs = typename STableRow<TSharedPtr<TListItemType>>::FArguments;
+			SMultiColumnTableRow<TSharedPtr<TListItemType>>::Construct(
+				FTableRowArgs()
+				.Style(FAppStyle::Get(), "TableView.AlternatingRow"),
+				InOwner
+				);
 		}
 
 		/** Generates the widget representing this row. */
diff --git a/Engine/Plugins/Developer/Concert/ConcertUI/ConcertSharedSlate/Source/ConcertSharedSlate/Private/Replication/Editor/View/Tree/SReplicationTreeView.h b/Engine/Plugins/Developer/Concert/ConcertUI/ConcertSharedSlate/Source/ConcertSharedSlate/Private/Replication/Editor/View/Tree/SReplicationTreeView.h
--- a/Engine/Plugins/Developer/Concert/ConcertUI/ConcertSharedSlate/Source/ConcertSharedSlate/Private/Replication/Editor/View/Tree/SReplicationTreeView.h
+++ b/Engine/Plugins/Developer/Concert/ConcertUI/ConcertSharedSlate/Source/ConcertSharedSlate/Private/Replication/Editor/View/Tree/SReplicationTreeView.h
@@ -383,6 +383,8 @@ namespace UE::ConcertSharedSlate
 					.SelectionMode(InArgs._SelectionMode)
 					.AllowOverscroll(EAllowOverscroll::No)
 					.HeaderRow(CreateHeaderRow(InArgs))
+					// Preserve the selection when the selected item is hidden due to a parent collapsing
+					.AllowInvisibleItemSelection(true)
 				]
 
 				+SVerticalBox::Slot()
