Filter improvements:
- Related filters have been combined, e.g. int32, uint32, etc. into an Integers group
- When no filters are selected, all properties are shown

diff --git a/Engine/Plugins/Developer/Concert/ConcertUI/ConcertSharedSlate/Source/ConcertSharedSlate/Private/Replication/Editor/View/ObjectViewer/Property/PropertyFilterBase.h b/Engine/Plugins/Developer/Concert/ConcertUI/ConcertSharedSlate/Source/ConcertSharedSlate/Private/Replication/Editor/View/ObjectViewer/Property/PropertyFilterBase.h
--- a/Engine/Plugins/Developer/Concert/ConcertUI/ConcertSharedSlate/Source/ConcertSharedSlate/Private/Replication/Editor/View/ObjectViewer/Property/PropertyFilterBase.h
+++ b/Engine/Plugins/Developer/Concert/ConcertUI/ConcertSharedSlate/Source/ConcertSharedSlate/Private/Replication/Editor/View/ObjectViewer/Property/PropertyFilterBase.h
@@ -22,13 +22,11 @@ namespace UE::ConcertSharedSlate
 	class FPropertyFilterBase : public IFilter<const FReplicatedPropertyData&>
 	{
 	public:
-
-
+		
 		//~ Begin IFilter Interface
 		virtual bool PassesFilter(const FReplicatedPropertyData& InItem) const final override
 		{
-			// If the property matches our filter criteria, it should NOT be shown in the UI since we're an inverse filter, so negate.
-			return !MatchesFilteredForProperty(InItem);
+			return MatchesFilteredForProperty(InItem);
 		}
 		virtual FChangedEvent& OnChanged() final override { return ChangedEventDelegate; }
 		//~ End IFilter Interface
diff --git a/Engine/Plugins/Developer/Concert/ConcertUI/ConcertSharedSlate/Source/ConcertSharedSlate/Private/Replication/Editor/View/ObjectViewer/Property/PropertyFrontendFilter.h b/Engine/Plugins/Developer/Concert/ConcertUI/ConcertSharedSlate/Source/ConcertSharedSlate/Private/Replication/Editor/View/ObjectViewer/Property/PropertyFrontendFilter.h
--- a/Engine/Plugins/Developer/Concert/ConcertUI/ConcertSharedSlate/Source/ConcertSharedSlate/Private/Replication/Editor/View/ObjectViewer/Property/PropertyFrontendFilter.h
+++ b/Engine/Plugins/Developer/Concert/ConcertUI/ConcertSharedSlate/Source/ConcertSharedSlate/Private/Replication/Editor/View/ObjectViewer/Property/PropertyFrontendFilter.h
@@ -26,7 +26,7 @@ namespace UE::ConcertSharedSlate
 		virtual FText GetToolTipText() const override { return Tooltip; }
 		virtual FLinearColor GetColor() const override { return FLinearColor(0.6f, 0.6f, 0.6f, 0.6f); }
 		virtual FName GetIconName() const override { return NAME_None; }
-		virtual bool IsInverseFilter() const override { return true; }
+		virtual bool IsInverseFilter() const override { return false; }
 		virtual void ActiveStateChanged(bool bActive) override {}
 		virtual void ModifyContextMenu(FMenuBuilder& MenuBuilder) override {}
 		virtual void SaveSettings(const FString& IniFilename, const FString& IniSection, const FString& SettingsString) const override {}
diff --git a/Engine/Plugins/Developer/Concert/ConcertUI/ConcertSharedSlate/Source/ConcertSharedSlate/Private/Replication/Editor/View/ObjectViewer/Property/SPropertyTreeView.cpp b/Engine/Plugins/Developer/Concert/ConcertUI/ConcertSharedSlate/Source/ConcertSharedSlate/Private/Replication/Editor/View/ObjectViewer/Property/SPropertyTreeView.cpp
--- a/Engine/Plugins/Developer/Concert/ConcertUI/ConcertSharedSlate/Source/ConcertSharedSlate/Private/Replication/Editor/View/ObjectViewer/Property/SPropertyTreeView.cpp
+++ b/Engine/Plugins/Developer/Concert/ConcertUI/ConcertSharedSlate/Source/ConcertSharedSlate/Private/Replication/Editor/View/ObjectViewer/Property/SPropertyTreeView.cpp
@@ -95,10 +95,8 @@ namespace UE::ConcertSharedSlate
 		// Do this AFTER initializing ReplicatedProperties because it triggers the OnItemsChanged callback.
 		for (const FFilterRef& Filter : Filters.DisabledByDefault)
 		{
-			FilterBar->SetFilterCheckState(Filter, ECheckBoxState::Checked);
+			FilterBar->SetFilterCheckState(Filter, ECheckBoxState::Unchecked);
 		}
-		// Our filters are inverse (FFilterBase::IsInverseFilter) so that means these disabled filters will be active now (so e.g. bool will be filtered out by default now)
-		FilterBar->DisableAllFilters();
 		// Show all the other filters as enabled (not greyed out: blue) - they will not run their logic since they are inverse.
 		for (const FFilterRef& Filter : Filters.EnabledByDefault)
 		{
@@ -123,52 +121,37 @@ namespace UE::ConcertSharedSlate
 
 	SPropertyTreeView::FBuildFilterBarResult SPropertyTreeView::BuildFilterBar()
 	{
-		TSharedRef<FFilterCategory> GeneralCategory = MakeShared<FFilterCategory>(
-			LOCTEXT("CommonCategory.Name", "General"),
-			LOCTEXT("CommonCategory.ToolTip", "Exclude general properties.")
-		);
-		TSharedRef<FFilterCategory> NumericCategory = MakeShared<FFilterCategory>(
-			LOCTEXT("NumericCategory.Name", "Numeric"),
-			LOCTEXT("NumericCategory.ToolTip", "Excludes numeric properties.")
+		TSharedRef<FFilterCategory> CommonCategory = MakeShared<FFilterCategory>(
+			LOCTEXT("CommonCategory.Name", "Common"),
+			LOCTEXT("CommonCategory.ToolTip", "Include commonly replicated properties.")
 		);
-		TSharedRef<FFilterCategory> ContainersCategory = MakeShared<FFilterCategory>(
-			LOCTEXT("ContainersCategory.Name", "Containers"),
-			LOCTEXT("ContainersCategory.ToolTip", "Exclude container properties.")
+		TSharedRef<FFilterCategory> UncommonCategory = MakeShared<FFilterCategory>(
+			LOCTEXT("NumericCategory.Advanced", "Uncommon"),
+			LOCTEXT("NumericCategory.ToolTip", "Include uncommonly replicated properties.")
 		);
 
 		FBuildFilterBarResult Result;
 		using FFrontendFilter = TPropertyFrontendFilter<FPropertyFilter_ByPropertyType>;
+		// Do not show under search bar
 		Result.DisabledByDefault =
 		{
-			MakeShared<FFrontendFilter>(MoveTemp(GeneralCategory), LOCTEXT("Bool", "Boolean"), TSet<FFieldClass*>{ FBoolProperty::StaticClass() }),
-			MakeShared<FFrontendFilter>(MoveTemp(NumericCategory), LOCTEXT("Byte", "Byte"), TSet<FFieldClass*>{ FByteProperty::StaticClass() }),
-			MakeShared<FFrontendFilter>(MoveTemp(GeneralCategory), LOCTEXT("Name", "Name"), TSet<FFieldClass*>{ FNameProperty::StaticClass() }),
-			MakeShared<FFrontendFilter>(MoveTemp(GeneralCategory), LOCTEXT("String", "String"), TSet<FFieldClass*>{ FStrProperty::StaticClass() }),
-			MakeShared<FFrontendFilter>(MoveTemp(GeneralCategory), LOCTEXT("Enum", "Enum"), TSet<FFieldClass*>{ FEnumProperty::StaticClass() })
+			MakeShared<FFrontendFilter>(MoveTemp(UncommonCategory), LOCTEXT("Bool", "Boolean"), TSet<FFieldClass*>{ FBoolProperty::StaticClass() }),
+			MakeShared<FFrontendFilter>(MoveTemp(UncommonCategory), LOCTEXT("Enum", "Enum"), TSet<FFieldClass*>{ FEnumProperty::StaticClass(), FByteProperty::StaticClass() }),
+			MakeShared<FFrontendFilter>(MoveTemp(UncommonCategory), LOCTEXT("Text", "Text"), LOCTEXT("Text.Tooltip", "Includes: FName, FString, FText"), TSet<FFieldClass*>{ FNameProperty::StaticClass(), FStrProperty::StaticClass(), FTextProperty::StaticClass() }),
+			MakeShared<FFrontendFilter>(MoveTemp(UncommonCategory), LOCTEXT("SoftPtr", "Soft Ptr"), TSet<FFieldClass*>{ FSoftObjectProperty::StaticClass() }),
 		};
+		// Show up under search bar as enabled
 		Result.EnabledByDefault =
 		{
-			MakeShared<FFrontendFilter>(MoveTemp(NumericCategory), LOCTEXT("Ints", "Integers"), LOCTEXT("Ints.Tooltip", "Includes integer types: uint16, uint32, uint64, int16, int32, int64"), TSet<FFieldClass*>{ FUInt16Property::StaticClass(), FUInt32Property::StaticClass(), FUInt64Property::StaticClass(), FInt16Property::StaticClass(), FIntProperty::StaticClass(), FInt64Property::StaticClass() }),
-			MakeShared<FFrontendFilter>(MoveTemp(NumericCategory), LOCTEXT("Float", "Float"), TSet<FFieldClass*>{ FFloatProperty::StaticClass() }),
-			MakeShared<FFrontendFilter>(MoveTemp(NumericCategory), LOCTEXT("Double", "Double"), TSet<FFieldClass*>{ FDoubleProperty::StaticClass() }),
-			MakeShared<FFrontendFilter>(MoveTemp(GeneralCategory), LOCTEXT("Struct", "Struct"), TSet<FFieldClass*>{ FStructProperty::StaticClass() }),
-			MakeShared<FFrontendFilter>(MoveTemp(ContainersCategory), LOCTEXT("Array", "Array"), TSet<FFieldClass*>{ FArrayProperty::StaticClass() }),
-			MakeShared<FFrontendFilter>(MoveTemp(ContainersCategory), LOCTEXT("Set", "Set"), TSet<FFieldClass*>{ FSetProperty::StaticClass() }),
-			MakeShared<FFrontendFilter>(MoveTemp(ContainersCategory), LOCTEXT("Map", "Map"), TSet<FFieldClass*>{ FMapProperty::StaticClass() }),
+			MakeShared<FFrontendFilter>(MoveTemp(CommonCategory), LOCTEXT("Ints", "Integer"), LOCTEXT("Ints.Tooltip", "Includes: uint16, uint32, uint64, int16, int32, int64"), TSet<FFieldClass*>{ FUInt16Property::StaticClass(), FUInt32Property::StaticClass(), FUInt64Property::StaticClass(), FInt16Property::StaticClass(), FIntProperty::StaticClass(), FInt64Property::StaticClass() }),
+			MakeShared<FFrontendFilter>(MoveTemp(CommonCategory), LOCTEXT("Floats", "Float"), LOCTEXT("Ints.Tooltip", "Includes: float, double"), TSet<FFieldClass*>{ FFloatProperty::StaticClass(), FDoubleProperty::StaticClass() }),
+			MakeShared<FFrontendFilter>(MoveTemp(CommonCategory), LOCTEXT("Struct", "Struct"), TSet<FFieldClass*>{ FStructProperty::StaticClass() }),
+			MakeShared<FFrontendFilter>(MoveTemp(CommonCategory), LOCTEXT("Containers", "Containers"), LOCTEXT("Ints.Tooltip", "Includes: array, set, map"), TSet<FFieldClass*>{ FArrayProperty::StaticClass(), FSetProperty::StaticClass(), FMapProperty::StaticClass() }),
 		};
-		
-		TArray<FFilterRef> AllFilters =
-		{
-			MakeShared<FFrontendFilter>(MoveTemp(NumericCategory), LOCTEXT("UInt16", "UInt16"), TSet<FFieldClass*>{ FUInt16Property::StaticClass() }),
-			MakeShared<FFrontendFilter>(MoveTemp(NumericCategory), LOCTEXT("UInt32", "UInt32"), TSet<FFieldClass*>{ FUInt32Property::StaticClass() }),
-			MakeShared<FFrontendFilter>(MoveTemp(NumericCategory), LOCTEXT("UInt64", "UInt64"), TSet<FFieldClass*>{ FUInt64Property::StaticClass() }),
-			MakeShared<FFrontendFilter>(MoveTemp(NumericCategory), LOCTEXT("Int16", "Int16"), TSet<FFieldClass*>{ FInt16Property::StaticClass() }),
-			MakeShared<FFrontendFilter>(MoveTemp(NumericCategory), LOCTEXT("Int32", "Int32"), TSet<FFieldClass*>{ FIntProperty::StaticClass() }),
-			MakeShared<FFrontendFilter>(MoveTemp(NumericCategory), LOCTEXT("Int64", "Int64"), TSet<FFieldClass*>{ FInt64Property::StaticClass() }),
-			MakeShared<FFrontendFilter>(MoveTemp(GeneralCategory), LOCTEXT("SoftPtr", "Soft Ptr"), TSet<FFieldClass*>{ FSoftObjectProperty::StaticClass() }),
-		};
-		AllFilters.Append(Result.DisabledByDefault);
+		// Show up in menu
+		TArray<FFilterRef> AllFilters;
 		AllFilters.Append(Result.EnabledByDefault);
+		AllFilters.Append(Result.DisabledByDefault);
 		
 		FilterBar = SNew(SReplicationFilterBar, MoveTemp(AllFilters))
 			.OnFilterChanged(this, &SPropertyTreeView::OnItemsChanged);
@@ -178,7 +161,21 @@ namespace UE::ConcertSharedSlate
 
 	bool SPropertyTreeView::PassesFilters(const TSharedPtr<FReplicatedPropertyData>& ReplicatedPropertyData) const
 	{
-		return FilterBar->GetAllActiveFilters()->PassesAllFilters(ReplicatedPropertyData);
+		return FilterBar->GetAllActiveFilters()->Num() == 0 // Return all items when none enabled
+			|| PassesAnyFilters(ReplicatedPropertyData);
+	}
+
+	bool SPropertyTreeView::PassesAnyFilters(const TSharedPtr<FReplicatedPropertyData>& ReplicatedPropertyData) const
+	{
+		TSharedPtr<TFilterCollection<TSharedPtr<FReplicatedPropertyData>>> FilterCollection = FilterBar->GetAllActiveFilters();
+		for (int32 Index = 0; Index < FilterCollection->Num(); Index++)
+		{
+			if (FilterCollection->GetFilterAtIndex(Index)->PassesFilter(ReplicatedPropertyData))
+			{
+				return true;
+			}
+		}
+		return false;
 	}
 
 	FText SPropertyTreeView::GetAllFilteredText() const
diff --git a/Engine/Plugins/Developer/Concert/ConcertUI/ConcertSharedSlate/Source/ConcertSharedSlate/Private/Replication/Editor/View/ObjectViewer/Property/SPropertyTreeView.h b/Engine/Plugins/Developer/Concert/ConcertUI/ConcertSharedSlate/Source/ConcertSharedSlate/Private/Replication/Editor/View/ObjectViewer/Property/SPropertyTreeView.h
--- a/Engine/Plugins/Developer/Concert/ConcertUI/ConcertSharedSlate/Source/ConcertSharedSlate/Private/Replication/Editor/View/ObjectViewer/Property/SPropertyTreeView.h
+++ b/Engine/Plugins/Developer/Concert/ConcertUI/ConcertSharedSlate/Source/ConcertSharedSlate/Private/Replication/Editor/View/ObjectViewer/Property/SPropertyTreeView.h
@@ -89,6 +89,7 @@ namespace UE::ConcertSharedSlate
 
 		/** Runs all filters through this item */
 		bool PassesFilters(const TSharedPtr<FReplicatedPropertyData>& ReplicatedPropertyData) const;
+		bool PassesAnyFilters(const TSharedPtr<FReplicatedPropertyData>& ReplicatedPropertyData) const;
 
 		/** Gets the message to display when all properties are filtered out. */
 		FText GetAllFilteredText() const;
