Make the replication tree view look and behave more like the scene outliner
- rows are height 22
- preserve the selection when the selected item is hidden due to a parent collapsing
- reduce font size of property section and all combo buttons
- align type columns so they form a vertical line

diff --git a/Engine/Plugins/Developer/Concert/ConcertApp/MultiUserClient/Source/MultiUserClient/Private/MultiUserReplicationStyle.cpp b/Engine/Plugins/Developer/Concert/ConcertApp/MultiUserClient/Source/MultiUserClient/Private/MultiUserReplicationStyle.cpp
--- a/Engine/Plugins/Developer/Concert/ConcertApp/MultiUserClient/Source/MultiUserClient/Private/MultiUserReplicationStyle.cpp
+++ b/Engine/Plugins/Developer/Concert/ConcertApp/MultiUserClient/Source/MultiUserClient/Private/MultiUserReplicationStyle.cpp
@@ -41,7 +41,7 @@ namespace UE::MultiUserClient
 
 		// Column widths
 		StyleSet->Set("AllClients.Object.OwnerColumnWidth", 200.f);
-		StyleSet->Set("AllClients.Object.ReplicationToggle", 50.f);
+		StyleSet->Set("AllClients.Object.ReplicationToggle", 45.f);
 		StyleSet->Set("AllClients.Property.OwnerColumnWidth", 200.f);
 		StyleSet->Set("SingleClient.Object.OwnerColumnWidth", 200.f);
 		StyleSet->Set("SingleClient.Property.OwnerColumnWidth", 200.f);
diff --git a/Engine/Plugins/Developer/Concert/ConcertApp/MultiUserClient/Source/MultiUserClient/Private/Widgets/ActiveSession/Replication/Client/Multi/Columns/SAssignPropertyComboBox.cpp b/Engine/Plugins/Developer/Concert/ConcertApp/MultiUserClient/Source/MultiUserClient/Private/Widgets/ActiveSession/Replication/Client/Multi/Columns/SAssignPropertyComboBox.cpp
--- a/Engine/Plugins/Developer/Concert/ConcertApp/MultiUserClient/Source/MultiUserClient/Private/Widgets/ActiveSession/Replication/Client/Multi/Columns/SAssignPropertyComboBox.cpp
+++ b/Engine/Plugins/Developer/Concert/ConcertApp/MultiUserClient/Source/MultiUserClient/Private/Widgets/ActiveSession/Replication/Client/Multi/Columns/SAssignPropertyComboBox.cpp
@@ -17,6 +17,7 @@
 #include "ScopedTransaction.h"
 #include "Styling/SlateTypes.h"
 #include "Widgets/Input/SComboButton.h"
+#include "Styling/AppStyle.h"
 
 #define LOCTEXT_NAMESPACE "SAssignPropertyComboBox"
 
@@ -81,10 +82,10 @@ namespace UE::MultiUserClient
 		[
 			SNew(SComboButton)
 			.HasDownArrow(true)
-			.ContentPadding(FMargin(2.0f, 2.0f))
 			.ButtonContent()
 			[
 				SAssignNew(ClientListWidget, ConcertClientSharedSlate::SHorizontalClientList, ConcertClient.ToSharedRef())
+				.Font(FAppStyle::GetFontStyle("PropertyWindow.NormalFont"))
 				.HighlightText_Lambda([this](){ return HighlightText ? *HighlightText : FText::GetEmpty(); })
 				.EmptyListSlot() [ SNew(SNoClients) ]
 			]
@@ -113,9 +114,11 @@ namespace UE::MultiUserClient
 			if (bIsLocalClient)
 			{
 				return SNew(SLocalClientName, ConcertClient.ToSharedRef())
+					.Font(FAppStyle::GetFontStyle("PropertyWindow.NormalFont"))
 					.HighlightText_Lambda([this](){ return HighlightText ? *HighlightText : FText::GetEmpty(); });
 			}
 			return SNew(SRemoteClientName, ConcertClient.ToSharedRef())
+				.Font(FAppStyle::GetFontStyle("PropertyWindow.NormalFont"))
 				.ClientEndpointId(EndpointId)
 				.HighlightText_Lambda([this](){ return HighlightText ? *HighlightText : FText::GetEmpty(); });
 		};
diff --git a/Engine/Plugins/Developer/Concert/ConcertApp/MultiUserClient/Source/MultiUserClient/Private/Widgets/ActiveSession/Replication/Client/Multi/Columns/SReassignObjectComboBox.cpp b/Engine/Plugins/Developer/Concert/ConcertApp/MultiUserClient/Source/MultiUserClient/Private/Widgets/ActiveSession/Replication/Client/Multi/Columns/SReassignObjectComboBox.cpp
--- a/Engine/Plugins/Developer/Concert/ConcertApp/MultiUserClient/Source/MultiUserClient/Private/Widgets/ActiveSession/Replication/Client/Multi/Columns/SReassignObjectComboBox.cpp
+++ b/Engine/Plugins/Developer/Concert/ConcertApp/MultiUserClient/Source/MultiUserClient/Private/Widgets/ActiveSession/Replication/Client/Multi/Columns/SReassignObjectComboBox.cpp
@@ -17,6 +17,7 @@
 #include "Textures/SlateIcon.h"
 #include "Widgets/Images/SThrobber.h"
 #include "Widgets/Input/SComboBox.h"
+#include "Widgets/Layout/SScaleBox.h"
 
 #define LOCTEXT_NAMESPACE "SReassignObjectComboBox"
 
@@ -104,10 +105,10 @@ namespace UE::MultiUserClient
 			[
 				SNew(SComboButton)
 				.HasDownArrow(true)
-				.ContentPadding(FMargin(2.0f, 2.0f))
 				.ButtonContent()
 				[
 					SAssignNew(ComboClientList, ConcertClientSharedSlate::SHorizontalClientList, InConcertClient)
+					.Font(FAppStyle::GetFontStyle("PropertyWindow.NormalFont"))
 					.HighlightText_Lambda([this](){ return HighlightText ? *HighlightText : FText::GetEmpty(); })
 					.EmptyListSlot() [ SNew(SNoClients) ]
 				]
diff --git a/Engine/Plugins/Developer/Concert/ConcertApp/MultiUserClient/Source/MultiUserClient/Private/Widgets/ActiveSession/Replication/Client/Multi/Columns/SReplicationMultiToggleCheckbox.cpp b/Engine/Plugins/Developer/Concert/ConcertApp/MultiUserClient/Source/MultiUserClient/Private/Widgets/ActiveSession/Replication/Client/Multi/Columns/SReplicationMultiToggleCheckbox.cpp
--- a/Engine/Plugins/Developer/Concert/ConcertApp/MultiUserClient/Source/MultiUserClient/Private/Widgets/ActiveSession/Replication/Client/Multi/Columns/SReplicationMultiToggleCheckbox.cpp
+++ b/Engine/Plugins/Developer/Concert/ConcertApp/MultiUserClient/Source/MultiUserClient/Private/Widgets/ActiveSession/Replication/Client/Multi/Columns/SReplicationMultiToggleCheckbox.cpp
@@ -38,27 +38,25 @@ namespace UE::MultiUserClient
 		
 		ChildSlot
 		[
-			SNew(SComboButton)
-			.HasDownArrow(true)
-			.ComboButtonStyle(FAppStyle::Get(), "SimpleComboButton")
-			.ForegroundColor(FSlateColor::UseStyle())
-			.ContentPadding(FMargin(2.0f, 2.0f))
-			.ButtonContent()
+			SNew(SHorizontalBox)
+			
+			+SHorizontalBox::Slot()
+			.AutoWidth()
 			[
-				SNew(SHorizontalBox)
-				.ToolTipText(this, &SReplicationMultiToggleCheckbox::GetRootToolTipText)
-
-				+SHorizontalBox::Slot()
-				.AutoWidth()
-				[
-					SNew(SCheckBox)
-					.IsChecked(this, &SReplicationMultiToggleCheckbox::GetCheckboxStateForThisAndChildren)
-					.IsEnabled(this, &SReplicationMultiToggleCheckbox::CanToggleThisOrChildren)
-					.OnCheckStateChanged(this, &SReplicationMultiToggleCheckbox::OnCheckboxStateChanged)
-				]
+				SNew(SCheckBox)
+				.IsChecked(this, &SReplicationMultiToggleCheckbox::GetCheckboxStateForThisAndChildren)
+				.IsEnabled(this, &SReplicationMultiToggleCheckbox::CanToggleThisOrChildren)
+				.OnCheckStateChanged(this, &SReplicationMultiToggleCheckbox::OnCheckboxStateChanged)
+			]
 			
-				+SHorizontalBox::Slot()
-				.AutoWidth()
+			+SHorizontalBox::Slot()
+			.AutoWidth()
+			[
+				SNew(SComboButton)
+				.HasDownArrow(true)
+				.ComboButtonStyle(FAppStyle::Get(), "SimpleComboButton")
+				.ForegroundColor(FSlateColor::UseStyle())
+				.ButtonContent()
 				[
 					SNew(SScaleBox)
 					.Stretch(EStretch::ScaleToFit)
@@ -69,8 +67,8 @@ namespace UE::MultiUserClient
 						.Image(FAppStyle::GetBrush("Icons.WarningWithColor"))
 					]
 				]
+				.OnGetMenuContent(this, &SReplicationMultiToggleCheckbox::GetDropDownMenuContent)
 			]
-			.OnGetMenuContent(this, &SReplicationMultiToggleCheckbox::GetDropDownMenuContent)
 		];
 	}
 
@@ -178,7 +176,7 @@ namespace UE::MultiUserClient
 				)
 			);
 		MenuBuilder.AddMenuEntry(
-			LOCTEXT("ThisAndChildren", "Toggle children only"),
+			LOCTEXT("ChildrenOnly", "Toggle children only"),
 			FText::GetEmpty(),
 			FSlateIcon(),
 			FUIAction(
@@ -187,7 +185,7 @@ namespace UE::MultiUserClient
 				)
 			);
 		MenuBuilder.AddMenuEntry(
-			LOCTEXT("OnlyThis", "Toggle this & children"),
+			LOCTEXT("ThisAndChildren", "Toggle this & children"),
 			FText::GetEmpty(),
 			FSlateIcon(),
 			FUIAction(
diff --git a/Engine/Plugins/Developer/Concert/ConcertUI/ConcertSharedSlate/Source/ConcertSharedSlate/Private/ConcertFrontendStyle.cpp b/Engine/Plugins/Developer/Concert/ConcertUI/ConcertSharedSlate/Source/ConcertSharedSlate/Private/ConcertFrontendStyle.cpp
--- a/Engine/Plugins/Developer/Concert/ConcertUI/ConcertSharedSlate/Source/ConcertSharedSlate/Private/ConcertFrontendStyle.cpp
+++ b/Engine/Plugins/Developer/Concert/ConcertUI/ConcertSharedSlate/Source/ConcertSharedSlate/Private/ConcertFrontendStyle.cpp
@@ -147,6 +147,11 @@ void FConcertFrontendStyle::Initialize()
 
 	// Colors
 	StyleSet->Set("Concert.DisconnectedColor", FLinearColor(0.672f, 0.672f, 0.672f));
+
+	// Floats
+	StyleSet->Set("Concert.Replication.Tree.RowHeight", 22.f);
+	StyleSet->Set("Concert.Replication.Tree.Object.LabelRowWidth", 255.f);
+	StyleSet->Set("Concert.Replication.Tree.Property.LabelRowWidth", 300.f);
 	
 	FSlateStyleRegistry::RegisterSlateStyle(*StyleSet.Get());
 };
diff --git a/Engine/Plugins/Developer/Concert/ConcertUI/ConcertSharedSlate/Source/ConcertSharedSlate/Private/Replication/Editor/View/ObjectViewer/Property/SPropertyTreeView.cpp b/Engine/Plugins/Developer/Concert/ConcertUI/ConcertSharedSlate/Source/ConcertSharedSlate/Private/Replication/Editor/View/ObjectViewer/Property/SPropertyTreeView.cpp
--- a/Engine/Plugins/Developer/Concert/ConcertUI/ConcertSharedSlate/Source/ConcertSharedSlate/Private/Replication/Editor/View/ObjectViewer/Property/SPropertyTreeView.cpp
+++ b/Engine/Plugins/Developer/Concert/ConcertUI/ConcertSharedSlate/Source/ConcertSharedSlate/Private/Replication/Editor/View/ObjectViewer/Property/SPropertyTreeView.cpp
@@ -126,8 +126,8 @@ namespace UE::ConcertSharedSlate
 			LOCTEXT("CommonCategory.ToolTip", "Include commonly replicated properties.")
 		);
 		TSharedRef<FFilterCategory> UncommonCategory = MakeShared<FFilterCategory>(
-			LOCTEXT("NumericCategory.Advanced", "Uncommon"),
-			LOCTEXT("NumericCategory.ToolTip", "Include uncommonly replicated properties.")
+			LOCTEXT("UncommonCategory.Advanced", "Uncommon"),
+			LOCTEXT("UncommonCategory.ToolTip", "Include uncommonly replicated properties.")
 		);
 
 		FBuildFilterBarResult Result;
@@ -135,6 +135,14 @@ namespace UE::ConcertSharedSlate
 		// Do not show under search bar
 		Result.DisabledByDefault =
 		{
+			// Ordering matters for the drop-down menu next to settings
+			// Common
+			MakeShared<FFrontendFilter>(MoveTemp(CommonCategory), LOCTEXT("Ints", "Integer"), LOCTEXT("Ints.Tooltip", "Includes: uint16, uint32, uint64, int16, int32, int64"), TSet<FFieldClass*>{ FUInt16Property::StaticClass(), FUInt32Property::StaticClass(), FUInt64Property::StaticClass(), FInt16Property::StaticClass(), FIntProperty::StaticClass(), FInt64Property::StaticClass() }),
+			MakeShared<FFrontendFilter>(MoveTemp(CommonCategory), LOCTEXT("Floats", "Float"), LOCTEXT("Floats.Tooltip", "Includes: float, double"), TSet<FFieldClass*>{ FFloatProperty::StaticClass(), FDoubleProperty::StaticClass() }),
+			MakeShared<FFrontendFilter>(MoveTemp(CommonCategory), LOCTEXT("Struct", "Struct"), TSet<FFieldClass*>{ FStructProperty::StaticClass() }),
+			MakeShared<FFrontendFilter>(MoveTemp(CommonCategory), LOCTEXT("Containers", "Containers"), LOCTEXT("Containers.Tooltip", "Includes: array, set, map"), TSet<FFieldClass*>{ FArrayProperty::StaticClass(), FSetProperty::StaticClass(), FMapProperty::StaticClass() }),
+
+			// Uncommon
 			MakeShared<FFrontendFilter>(MoveTemp(UncommonCategory), LOCTEXT("Bool", "Boolean"), TSet<FFieldClass*>{ FBoolProperty::StaticClass() }),
 			MakeShared<FFrontendFilter>(MoveTemp(UncommonCategory), LOCTEXT("Enum", "Enum"), TSet<FFieldClass*>{ FEnumProperty::StaticClass(), FByteProperty::StaticClass() }),
 			MakeShared<FFrontendFilter>(MoveTemp(UncommonCategory), LOCTEXT("Text", "Text"), LOCTEXT("Text.Tooltip", "Includes: FName, FString, FText"), TSet<FFieldClass*>{ FNameProperty::StaticClass(), FStrProperty::StaticClass(), FTextProperty::StaticClass() }),
@@ -143,10 +151,7 @@ namespace UE::ConcertSharedSlate
 		// Show up under search bar as enabled
 		Result.EnabledByDefault =
 		{
-			MakeShared<FFrontendFilter>(MoveTemp(CommonCategory), LOCTEXT("Ints", "Integer"), LOCTEXT("Ints.Tooltip", "Includes: uint16, uint32, uint64, int16, int32, int64"), TSet<FFieldClass*>{ FUInt16Property::StaticClass(), FUInt32Property::StaticClass(), FUInt64Property::StaticClass(), FInt16Property::StaticClass(), FIntProperty::StaticClass(), FInt64Property::StaticClass() }),
-			MakeShared<FFrontendFilter>(MoveTemp(CommonCategory), LOCTEXT("Floats", "Float"), LOCTEXT("Ints.Tooltip", "Includes: float, double"), TSet<FFieldClass*>{ FFloatProperty::StaticClass(), FDoubleProperty::StaticClass() }),
-			MakeShared<FFrontendFilter>(MoveTemp(CommonCategory), LOCTEXT("Struct", "Struct"), TSet<FFieldClass*>{ FStructProperty::StaticClass() }),
-			MakeShared<FFrontendFilter>(MoveTemp(CommonCategory), LOCTEXT("Containers", "Containers"), LOCTEXT("Ints.Tooltip", "Includes: array, set, map"), TSet<FFieldClass*>{ FArrayProperty::StaticClass(), FSetProperty::StaticClass(), FMapProperty::StaticClass() }),
+			// We'll not enable any filters by default because that's the default for other places in the engine, like the Content Browser, too
 		};
 		// Show up in menu
 		TArray<FFilterRef> AllFilters;
diff --git a/Engine/Plugins/Developer/Concert/ConcertUI/ConcertSharedSlate/Source/ConcertSharedSlate/Private/Replication/Editor/View/ObjectViewer/SReplicationStreamViewer.cpp b/Engine/Plugins/Developer/Concert/ConcertUI/ConcertSharedSlate/Source/ConcertSharedSlate/Private/Replication/Editor/View/ObjectViewer/SReplicationStreamViewer.cpp
--- a/Engine/Plugins/Developer/Concert/ConcertUI/ConcertSharedSlate/Source/ConcertSharedSlate/Private/Replication/Editor/View/ObjectViewer/SReplicationStreamViewer.cpp
+++ b/Engine/Plugins/Developer/Concert/ConcertUI/ConcertSharedSlate/Source/ConcertSharedSlate/Private/Replication/Editor/View/ObjectViewer/SReplicationStreamViewer.cpp
@@ -229,7 +229,8 @@ namespace UE::ConcertSharedSlate
 					ObjectViewOptions.MakeViewOptionsComboButton()
 				]
 			]
-			.NoItemsContent() [ SNew(STextBlock).Text(NoObjectsAttribute) ];
+			.NoItemsContent() [ SNew(STextBlock).Text(NoObjectsAttribute) ]
+			.RowStyle(FAppStyle::Get(), "TableView.AlternatingRow");
 	}
 
 	TSharedRef<SWidget> SReplicationStreamViewer::CreatePropertiesSection(const FArguments& InArgs)
diff --git a/Engine/Plugins/Developer/Concert/ConcertUI/ConcertSharedSlate/Source/ConcertSharedSlate/Private/Replication/Editor/View/SelectionViewerColumns.cpp b/Engine/Plugins/Developer/Concert/ConcertUI/ConcertSharedSlate/Source/ConcertSharedSlate/Private/Replication/Editor/View/SelectionViewerColumns.cpp
--- a/Engine/Plugins/Developer/Concert/ConcertUI/ConcertSharedSlate/Source/ConcertSharedSlate/Private/Replication/Editor/View/SelectionViewerColumns.cpp
+++ b/Engine/Plugins/Developer/Concert/ConcertUI/ConcertSharedSlate/Source/ConcertSharedSlate/Private/Replication/Editor/View/SelectionViewerColumns.cpp
@@ -2,14 +2,13 @@
 
 #include "Replication/Editor/View/SelectionViewerColumns.h"
 
+#include "ConcertFrontendStyle.h"
 #include "Replication/Editor/Model/IEditableReplicationStreamModel.h"
-#include "Replication/Editor/Model/Object/IObjectNameModel.h"
 #include "Replication/Editor/Model/ReplicatedPropertyData.h"
 #include "Replication/Editor/Model/ReplicatedObjectData.h"
 #include "Replication/Editor/View/DisplayUtils.h"
 #include "Replication/Editor/View/ObjectEditor/SBaseReplicationStreamEditor.h"
 #include "Replication/Editor/View/ReplicationColumnsUtils.h"
-#include "Replication/ObjectUtils.h"
 #include "Replication/PropertyChainUtils.h"
 
 #include "Internationalization/Internationalization.h"
@@ -18,6 +17,7 @@
 #include "Widgets/SBoxPanel.h"
 #include "Widgets/Text/STextBlock.h"
 #include "Widgets/Input/SCheckBox.h"
+#include "Styling/AppStyle.h"
 
 #define LOCTEXT_NAMESPACE "ReplicationObjectColumns"
 
@@ -74,7 +74,7 @@ namespace UE::ConcertSharedSlate::ReplicationColumns::TopLevel
 				.ColumnSortOrder(static_cast<int32>(ETopLevelColumnOrder::Label)),
 			SHeaderRow::Column(LabelColumnId)
 				.DefaultLabel(LOCTEXT("LabelColumnLabel", "Label"))
-				.FillWidth(1.f)
+				.FillSized(FConcertFrontendStyle::Get()->GetFloat("Concert.Replication.Tree.Object.LabelRowWidth"))
 			);
 	}
 	
@@ -120,6 +120,7 @@ namespace UE::ConcertSharedSlate::ReplicationColumns::Property
 				.GenerateWidgetColumn_Lambda([](const FReplicationPropertyColumn::FBuildArgs& Args)
 				{
 					return SNew(STextBlock)
+						.Font(FAppStyle::GetFontStyle("PropertyWindow.NormalFont"))
 						.HighlightText(TAttribute<FText>::CreateLambda([HighlightText = Args.HighlightText](){ return *HighlightText; }))
 						.Text(DisplayUtils::GetPropertyDisplayText(Args.RowData.GetProperty()));
 				})
@@ -134,7 +135,7 @@ namespace UE::ConcertSharedSlate::ReplicationColumns::Property
 				.ColumnSortOrder(static_cast<int32>(EReplicationPropertyColumnOrder::Label)),
 			SHeaderRow::Column(LabelColumnId)
 				.DefaultLabel(LOCTEXT("LabelColumnLabel", "Label"))
-				.FillWidth(1.f)
+				.FillSized(FConcertFrontendStyle::Get()->GetFloat("Concert.Replication.Tree.Property.LabelRowWidth"))
 			);
 	}
 	
@@ -152,6 +153,7 @@ namespace UE::ConcertSharedSlate::ReplicationColumns::Property
 				.GenerateWidgetColumn_Lambda([](const FReplicationPropertyColumn::FBuildArgs& Args)
 				{
 					return SNew(STextBlock)
+						.Font(FAppStyle::GetFontStyle("PropertyWindow.NormalFont"))
 						.HighlightText(TAttribute<FText>::CreateLambda([HighlightText = Args.HighlightText](){ return *HighlightText; }))
 						.Text(GetDisplayText(Args.RowData));
 				})
@@ -166,7 +168,7 @@ namespace UE::ConcertSharedSlate::ReplicationColumns::Property
 				.ColumnSortOrder(static_cast<int32>(EReplicationPropertyColumnOrder::Type)),
 			SHeaderRow::Column(TypeColumnId)
 				.DefaultLabel(LOCTEXT("TypeColumnLabel", "Type"))
-				.FillWidth(0.5f)
+				.FillWidth(1.f)
 			);
 	}
 }
diff --git a/Engine/Plugins/Developer/Concert/ConcertUI/ConcertSharedSlate/Source/ConcertSharedSlate/Private/Replication/Editor/View/Tree/SReplicationColumnRow.h b/Engine/Plugins/Developer/Concert/ConcertUI/ConcertSharedSlate/Source/ConcertSharedSlate/Private/Replication/Editor/View/Tree/SReplicationColumnRow.h
--- a/Engine/Plugins/Developer/Concert/ConcertUI/ConcertSharedSlate/Source/ConcertSharedSlate/Private/Replication/Editor/View/Tree/SReplicationColumnRow.h
+++ b/Engine/Plugins/Developer/Concert/ConcertUI/ConcertSharedSlate/Source/ConcertSharedSlate/Private/Replication/Editor/View/Tree/SReplicationColumnRow.h
@@ -2,12 +2,16 @@
 
 #pragma once
 
+#include "ConcertFrontendStyle.h"
 #include "Replication/Editor/View/ReplicationColumn.h"
+
 #include "Widgets/Layout/SBox.h"
 #include "Widgets/SNullWidget.h"
 #include "Widgets/Views/SExpanderArrow.h"
 #include "Widgets/Views/STableRow.h"
 #include "Widgets/Views/STableViewBase.h"
+#include "Styling/AppStyle.h"
+#include "Styling/CoreStyle.h"
 
 namespace UE::ConcertSharedSlate
 {
@@ -27,7 +31,7 @@ namespace UE::ConcertSharedSlate
 		DECLARE_DELEGATE_RetVal_TwoParams(TSharedPtr<SWidget>, FOverrideColumnWidget, const FName& ColumnName, const TListItemType& RowData);
 		
 		SLATE_BEGIN_ARGS(SReplicationColumnRow)
-			: _RowHeight(20.f)
+			: _Style(&FCoreStyle::Get().GetWidgetStyle<FTableRowStyle>("TableView.Row"))
 		{}
 			/** Used for highlighting the text being searched. */
 			SLATE_ARGUMENT(TSharedPtr<FText>, HighlightText)
@@ -47,8 +51,9 @@ namespace UE::ConcertSharedSlate
 			/** The name of the column that will have the SExpanderArrow for the tree view. */
 			SLATE_ARGUMENT(FName, ExpandableColumnLabel)
 
-			/** The height of the row */
-			SLATE_ARGUMENT(float, RowHeight)
+			/** Style to use for rows, e.g. for making them alternate in grey */
+			SLATE_STYLE_ARGUMENT(FTableRowStyle, Style)
+		
 		SLATE_END_ARGS()
 
 		void Construct(
@@ -60,12 +65,11 @@ namespace UE::ConcertSharedSlate
 			HighlightText = InArgs._HighlightText;
 			RowData = InArgs._RowData;
 			ExpandableColumnLabel = InArgs._ExpandableColumnLabel;
-			RowHeight = InArgs._RowHeight;
-			
+
 			using FTableRowArgs = typename STableRow<TSharedPtr<TListItemType>>::FArguments;
 			SMultiColumnTableRow<TSharedPtr<TListItemType>>::Construct(
 				FTableRowArgs()
-				.Style(FAppStyle::Get(), "TableView.AlternatingRow"),
+				.Style(InArgs._Style),
 				InOwner
 				);
 		}
@@ -92,7 +96,8 @@ namespace UE::ConcertSharedSlate
 			if (!bNeedsExpanderArrow)
 			{
 				return SNew(SBox)
-					.MinDesiredHeight(RowHeight)
+					// Enforce all items to be the same size so it is more consistent with other places, like the SSceneOutliner
+					.MaxDesiredHeight(FConcertFrontendStyle::Get()->GetFloat("Concert.Replication.Tree.RowHeight"))
 					.VAlign(VAlign_Center)
 					[
 						ColumnWidget
@@ -100,7 +105,8 @@ namespace UE::ConcertSharedSlate
 			}
 			
 			return SNew(SBox)
-				.MinDesiredHeight(RowHeight)
+				// Enforce all items to be the same size so it is more consistent with other places, like the SSceneOutliner
+				.MaxDesiredHeight(FConcertFrontendStyle::Get()->GetFloat("Concert.Replication.Tree.RowHeight"))
 				[
 					SNew(SHorizontalBox)
 
@@ -129,6 +135,5 @@ namespace UE::ConcertSharedSlate
 		TSharedPtr<FText> HighlightText;
 		TSharedPtr<TListItemType> RowData;
 		FName ExpandableColumnLabel;
-		float RowHeight;
 	};
 }
\ No newline at end of file
diff --git a/Engine/Plugins/Developer/Concert/ConcertUI/ConcertSharedSlate/Source/ConcertSharedSlate/Private/Replication/Editor/View/Tree/SReplicationTreeView.h b/Engine/Plugins/Developer/Concert/ConcertUI/ConcertSharedSlate/Source/ConcertSharedSlate/Private/Replication/Editor/View/Tree/SReplicationTreeView.h
--- a/Engine/Plugins/Developer/Concert/ConcertUI/ConcertSharedSlate/Source/ConcertSharedSlate/Private/Replication/Editor/View/Tree/SReplicationTreeView.h
+++ b/Engine/Plugins/Developer/Concert/ConcertUI/ConcertSharedSlate/Source/ConcertSharedSlate/Private/Replication/Editor/View/Tree/SReplicationTreeView.h
@@ -14,6 +14,8 @@
 #include "Widgets/Layout/SWidgetSwitcher.h"
 #include "Widgets/SBoxPanel.h"
 #include "Widgets/Views/STreeView.h"
+#include "Styling/AppStyle.h"
+#include "Styling/CoreStyle.h"
 
 #define LOCTEXT_NAMESPACE "SReplicationListView"
 
@@ -46,6 +48,7 @@ namespace UE::ConcertSharedSlate
 		SLATE_BEGIN_ARGS(SReplicationTreeView<TItemType>)
 			: _HeaderRowVisibility(EVisibility::Visible)
 			, _SelectionMode(ESelectionMode::Single)
+			, _RowStyle(&FCoreStyle::Get().GetWidgetStyle<FTableRowStyle>("TableView.Row"))
 		{}
 			/** The items to display */
 			SLATE_ARGUMENT(TArray<TSharedPtr<TItemType>>*, RootItemsSource)
@@ -96,6 +99,9 @@ namespace UE::ConcertSharedSlate
 
 			/** Optional, alternate content to show instead of the tree view when there are no rows. */
 			SLATE_NAMED_SLOT(FArguments, NoItemsContent)
+
+			/** Style to use for rows, e.g. for making them alternate in grey */
+			SLATE_STYLE_ARGUMENT(FTableRowStyle, RowStyle)
 		SLATE_END_ARGS()
 
 		void Construct(const FArguments& InArgs)
@@ -109,6 +115,7 @@ namespace UE::ConcertSharedSlate
 			OverrideColumnWidget = InArgs._OverrideColumnWidget;
 			IsSearchableItemDelegate = InArgs._IsSearchableItem;
 			ExpandableColumnId = InArgs._ExpandableColumnLabel;
+			RowStyle = InArgs._RowStyle;
 			
 			SearchText = MakeShared<FText>();
 			SearchTextFilter = MakeShared<TTextFilter<const TSharedPtr<TItemType>&>>(TTextFilter<const TSharedPtr<TItemType>&>::FItemToStringArray::CreateSP(this, &SReplicationTreeView::PopulateSearchStrings));
@@ -305,6 +312,9 @@ namespace UE::ConcertSharedSlate
 		/** Optional callback for determining whether this item can be filtered. If false, it will not be shown when searched. */
 		FIsSearchableItem IsSearchableItemDelegate;
 
+		/** Style to use for rows */
+		const FTableRowStyle* RowStyle = nullptr;
+
 		// STreeView creation
 		TSharedRef<SWidget> CreateTreeView(const FArguments& InArgs);
 		TSharedRef<SHeaderRow> CreateHeaderRow(const FArguments& InArgs);
@@ -446,7 +456,8 @@ namespace UE::ConcertSharedSlate
 			.ColumnGetter(ColumnGetter)
 			.OverrideColumnWidget(OverrideColumnWidget)
 			.RowData(Item)
-			.ExpandableColumnLabel(ExpandableColumnId);
+			.ExpandableColumnLabel(ExpandableColumnId)
+			.Style(RowStyle);
 	}
 
 	template <typename TItemType>
