Fixes to replication view
- Fix property view not refreshing properly
- When adding an actor, select added actor only instead of also all the automatically added subobjects

diff --git a/Engine/Plugins/Developer/Concert/ConcertUI/ConcertSharedSlate/Source/ConcertSharedSlate/Private/Replication/Editor/View/ObjectEditor/SBaseReplicationStreamEditor.cpp b/Engine/Plugins/Developer/Concert/ConcertUI/ConcertSharedSlate/Source/ConcertSharedSlate/Private/Replication/Editor/View/ObjectEditor/SBaseReplicationStreamEditor.cpp
--- a/Engine/Plugins/Developer/Concert/ConcertUI/ConcertSharedSlate/Source/ConcertSharedSlate/Private/Replication/Editor/View/ObjectEditor/SBaseReplicationStreamEditor.cpp
+++ b/Engine/Plugins/Developer/Concert/ConcertUI/ConcertSharedSlate/Source/ConcertSharedSlate/Private/Replication/Editor/View/ObjectEditor/SBaseReplicationStreamEditor.cpp
@@ -122,9 +122,22 @@ namespace UE::ConcertSharedSlate
 		if (!AddedObjects.IsEmpty())
 		{
 			TArray<FSoftObjectPath> TopLevelObjects;
-			Algo::Transform(AddedObjects, TopLevelObjects, [](const UObject* Object){ return Object; });
-			ReplicationViewer->SelectObjects(TopLevelObjects);
 			
+			// Goal: select an object so the property view immediately shows properties for some object.
+			// Problem: if objects have different classes, property view will be empty (incompatible class)
+			// Solution: Select the highest object in the hierarchy, which is usually an actor. Users usually select an actor to add in the "Add" combo button so this also makes intuitive sense.
+			// Caveat: There may be multiple hierarchies (e.g. if multiple actors were added). This is a very seldom case though: too bad.
+			Algo::TransformIf(AddedObjects, TopLevelObjects,
+				[this, &AddedObjects](const UObject* Object)
+				{
+					const TOptional<IObjectHierarchyModel::FParentInfo> ParentInfo = ObjectHierarchy->GetParentInfo(Object);
+					const bool bIsTopOfHierarchy = !ParentInfo || !AddedObjects.ContainsByPredicate([&ParentInfo](UObject* AddedObject){ return FSoftObjectPath(AddedObject) == ParentInfo->Parent; });
+					return bIsTopOfHierarchy;
+				},
+				[](const UObject* Object){ return Object; });
+			ReplicationViewer->SelectObjects(TopLevelObjects);
+
+			// Expand the hierarchy for all added objects for easier editing
 			constexpr bool bRecursive = true;
 			ReplicationViewer->ExpandObjects(TopLevelObjects, bRecursive);
 		}
diff --git a/Engine/Plugins/Developer/Concert/ConcertUI/ConcertSharedSlate/Source/ConcertSharedSlate/Private/Replication/Editor/View/ObjectViewer/SReplicatedPropertyView.cpp b/Engine/Plugins/Developer/Concert/ConcertUI/ConcertSharedSlate/Source/ConcertSharedSlate/Private/Replication/Editor/View/ObjectViewer/SReplicatedPropertyView.cpp
--- a/Engine/Plugins/Developer/Concert/ConcertUI/ConcertSharedSlate/Source/ConcertSharedSlate/Private/Replication/Editor/View/ObjectViewer/SReplicatedPropertyView.cpp
+++ b/Engine/Plugins/Developer/Concert/ConcertUI/ConcertSharedSlate/Source/ConcertSharedSlate/Private/Replication/Editor/View/ObjectViewer/SReplicatedPropertyView.cpp
@@ -10,10 +10,7 @@
 
 #include "Algo/AllOf.h"
 #include "Algo/ForEach.h"
-#include "Widgets/Layout/SBorder.h"
-#include "Widgets/Layout/SSplitter.h"
 #include "Widgets/Layout/SWidgetSwitcher.h"
-#include "Widgets/SBoxPanel.h"
 #include "Widgets/Text/STextBlock.h"
 
 #define LOCTEXT_NAMESPACE "SReplicatedPropertyView"
@@ -34,7 +31,7 @@ namespace UE::ConcertSharedSlate
 	
 	void SReplicatedPropertyView::RefreshPropertyData()
 	{
-		const TArray<FSoftObjectPath> SelectedObjects = GetObjectsSelectedForPropertyEditing();
+		TArray<FSoftObjectPath> SelectedObjects = GetObjectsSelectedForPropertyEditing();
 		if (SelectedObjects.IsEmpty())
 		{
 			SetPropertyContent(EReplicatedPropertyContent::NoSelection);
@@ -54,6 +51,17 @@ namespace UE::ConcertSharedSlate
 			return;
 		}
 
+		// If the objects have changed, the classes may share properties.
+		// In that case, below we'd reuse the item pointer, which would cause the tree view to re-use the old row widgets.
+		// However, we must regenerate all column widgets since they may be referencing the object the row was originally built for. So they'd display the state of the previous object still!
+		// Example: Assign property combo-box in Multi-User All Clients view displays who has the property assigned.
+		// Note: If the objects did not change, we definitely want to reuse item pointers since otherwise the user row selection is reset.
+		const bool bCanReusePropertyData = PreviousSelectedObjects == SelectedObjects; // This SHOULD be an order independent compare but usually Num == 1, so whatever
+		if (!bCanReusePropertyData)
+		{
+			ChainToPropertyDataCache.Reset();
+		}
+
 		// Build the set of properties that are shared by all of the selected objects
 		TSet<FConcertPropertyChain> SharedProperties = PropertiesModel->GetAllProperties(SelectedObjects[0]);
 		for (int32 i = 1; i < SelectedObjects.Num(); ++i)
@@ -79,6 +87,8 @@ namespace UE::ConcertSharedSlate
 		// The tree view requires the item source to only contain the root items.
 		BuildRootPropertyRowData();
 		SetPropertyContent(EReplicatedPropertyContent::Properties);
+
+		PreviousSelectedObjects = MoveTemp(SelectedObjects);
 		ReplicatedProperties->OnItemsChanged();
 	}
 
diff --git a/Engine/Plugins/Developer/Concert/ConcertUI/ConcertSharedSlate/Source/ConcertSharedSlate/Private/Replication/Editor/View/ObjectViewer/SReplicatedPropertyView.h b/Engine/Plugins/Developer/Concert/ConcertUI/ConcertSharedSlate/Source/ConcertSharedSlate/Private/Replication/Editor/View/ObjectViewer/SReplicatedPropertyView.h
--- a/Engine/Plugins/Developer/Concert/ConcertUI/ConcertSharedSlate/Source/ConcertSharedSlate/Private/Replication/Editor/View/ObjectViewer/SReplicatedPropertyView.h
+++ b/Engine/Plugins/Developer/Concert/ConcertUI/ConcertSharedSlate/Source/ConcertSharedSlate/Private/Replication/Editor/View/ObjectViewer/SReplicatedPropertyView.h
@@ -87,6 +87,9 @@ namespace UE::ConcertSharedSlate
 		/** Inverse map of PropertyRowData using FReplicatedPropertyData::GetProperty as key. Contains all elements of PropertyRowData. */
 		TMap<FConcertPropertyChain, TSharedPtr<FReplicatedPropertyData>> ChainToPropertyDataCache;
 
+		/** Used to determine whether to rebuild the entire property data. */
+		TArray<FSoftObjectPath> PreviousSelectedObjects;
+
 		TSharedRef<SWidget> CreatePropertiesView(const FArguments& InArgs);
 		
 		TSharedRef<FReplicatedPropertyData> AllocatePropertyData(FSoftClassPath OwningClass, FConcertPropertyChain PropertyChain);
