Fix removing last property removing entire hierarchy

diff --git a/Engine/Plugins/Developer/Concert/ConcertApp/MultiUserClient/Source/MultiUserClient/Private/Replication/Client/ReplicationClient.cpp b/Engine/Plugins/Developer/Concert/ConcertApp/MultiUserClient/Source/MultiUserClient/Private/Replication/Client/ReplicationClient.cpp
--- a/Engine/Plugins/Developer/Concert/ConcertApp/MultiUserClient/Source/MultiUserClient/Private/Replication/Client/ReplicationClient.cpp
+++ b/Engine/Plugins/Developer/Concert/ConcertApp/MultiUserClient/Source/MultiUserClient/Private/Replication/Client/ReplicationClient.cpp
@@ -89,7 +89,26 @@ namespace UE::MultiUserClient
 	void FReplicationClient::OnServerStateChanged()
 	{
 		// Whenever this client's server state changes, the UI must be refreshed.
+
+		// The UI adds empty actors. However, we never send them to the server...
+		TSet<FSoftObjectPath> StagedObjects;
+		for (const TPair<FSoftObjectPath, FConcertReplicatedObjectInfo>& Pair : GetClientContent()->Stream->ReplicationMap.ReplicatedObjects)
+		{
+			if (Pair.Value.PropertySelection.ReplicatedProperties.IsEmpty())
+			{
+				StagedObjects.Add(Pair.Key);
+			}
+		}
+		
+		// ... if the user removes the last property from the entire actor-component hierarchy, we want the hierarchy to continue to displayed...
 		GetClientContent()->Stream->ReplicationMap = GetStreamSynchronizer().GetServerState();
+		// ... so add back the staged objects
+		for (const FSoftObjectPath& StagedObject : StagedObjects)
+		{
+			GetClientContent()->Stream->ReplicationMap.ReplicatedObjects.Add(StagedObject);
+		}
+		// To remove the hierarchy, the user must click the actor and delete it explicitly, which will call IEditableReplicationStream::RemoveObjects on the staged objects.
+		
 		DeferOnModelChanged();
 	}
 
diff --git a/Engine/Plugins/Developer/Concert/ConcertUI/ConcertSharedSlate/Source/ConcertSharedSlate/Private/Replication/Editor/View/ObjectEditor/SBaseReplicationStreamEditor.cpp b/Engine/Plugins/Developer/Concert/ConcertUI/ConcertSharedSlate/Source/ConcertSharedSlate/Private/Replication/Editor/View/ObjectEditor/SBaseReplicationStreamEditor.cpp
--- a/Engine/Plugins/Developer/Concert/ConcertUI/ConcertSharedSlate/Source/ConcertSharedSlate/Private/Replication/Editor/View/ObjectEditor/SBaseReplicationStreamEditor.cpp
+++ b/Engine/Plugins/Developer/Concert/ConcertUI/ConcertSharedSlate/Source/ConcertSharedSlate/Private/Replication/Editor/View/ObjectEditor/SBaseReplicationStreamEditor.cpp
@@ -116,7 +116,7 @@ namespace UE::ConcertSharedSlate
 
 	void SBaseReplicationStreamEditor::OnObjectsChanged(TConstArrayView<UObject*> AddedObjects, TConstArrayView<FSoftObjectPath> RemovedObjects, EReplicatedObjectChangeReason ChangeReason)
 	{
-		ReplicationViewer->RefreshObjectData();
+		ReplicationViewer->RequestObjectDataRefresh();
 			
 		// Newly added objects should be automatically selected
 		if (!AddedObjects.IsEmpty())
@@ -145,7 +145,7 @@ namespace UE::ConcertSharedSlate
 
 	void SBaseReplicationStreamEditor::OnPropertiesChanged()
 	{
-		ReplicationViewer->RefreshPropertyData();
+		ReplicationViewer->RequestPropertyDataRefresh();
 	}
 
 	TSharedRef<SWidget> SBaseReplicationStreamEditor::BuildRootAddObjectWidgets()
diff --git a/Engine/Plugins/Developer/Concert/ConcertUI/ConcertSharedSlate/Source/ConcertSharedSlate/Private/Replication/Editor/View/ObjectViewer/SReplicationStreamViewer.cpp b/Engine/Plugins/Developer/Concert/ConcertUI/ConcertSharedSlate/Source/ConcertSharedSlate/Private/Replication/Editor/View/ObjectViewer/SReplicationStreamViewer.cpp
--- a/Engine/Plugins/Developer/Concert/ConcertUI/ConcertSharedSlate/Source/ConcertSharedSlate/Private/Replication/Editor/View/ObjectViewer/SReplicationStreamViewer.cpp
+++ b/Engine/Plugins/Developer/Concert/ConcertUI/ConcertSharedSlate/Source/ConcertSharedSlate/Private/Replication/Editor/View/ObjectViewer/SReplicationStreamViewer.cpp
@@ -37,8 +37,8 @@
 
 	void SReplicationStreamViewer::Refresh()
 	{
-		RefreshObjectData();
-		RefreshPropertyData();
+		RequestObjectDataRefresh();
+		RequestPropertyDataRefresh();
 	}
 
 	void SReplicationStreamViewer::RequestObjectColumnResort(const FName& ColumnId)
@@ -56,53 +56,6 @@
 		return PropertySection->GetObjectsSelectedForPropertyEditing();
 	}
 
-	void SReplicationStreamViewer::RefreshObjectData()
-	{
-		// Re-using existing instances is tricky: we cannot update the object path in an item because the list view will no detect this change;
-		// list view only looks at the shared ptr address. So the UI will not be refreshed. Since the number of items will be small, just reallocate... 
-		AllObjectRowData.Empty();
-
-		// Try to re-use old instances by using the old PathToObjectDataCache. This is also done so the expansion states restore correctly in the tree view.
-		TMap<FSoftObjectPath, TSharedPtr<FReplicatedObjectData>> NewPathToObjectDataCache;
-		
-		// Do a complete refresh.
-		// Complete refresh is acceptable because the list is updated infrequently and typically small < 500 items.
-		// An alternative would be to change RefreshObjectData to be called with two variables ObjectsAdded and ObjectsRemoved.
-		PropertiesModel->ForEachReplicatedObject([this, &NewPathToObjectDataCache](const FSoftObjectPath& ObjectPath) mutable
-		{
-			TOptional<IObjectHierarchyModel::FParentInfo> ParentInfo = ObjectHierarchy->GetParentInfo(ObjectPath);
-			const bool bIsActor = !ParentInfo; 
-			if (bIsActor || ShouldDisplayObject(ObjectPath, ParentInfo->Relationship))
-			{
-				const TSharedPtr<FReplicatedObjectData>* ExistingItem = PathToObjectDataCache.Find(ObjectPath);
-				ExistingItem = ExistingItem ? ExistingItem : NewPathToObjectDataCache.Find(ObjectPath);
-				const TSharedRef<FReplicatedObjectData> Item = ExistingItem ? ExistingItem->ToSharedRef() : AllocateObjectData(ObjectPath);
-				AllObjectRowData.AddUnique(Item);
-				NewPathToObjectDataCache.Emplace(ObjectPath, Item);
-				
-				BuildObjectHierarchyIfNeeded(Item, NewPathToObjectDataCache);
-			}
-			
-			return EBreakBehavior::Continue;
-		});
-
-		// Only refresh the tree if it is necessary as it causes us to select stuff in the subobject view
-		if (!PathToObjectDataCache.OrderIndependentCompareEqual(NewPathToObjectDataCache))
-		{
-			// If an item was removed, then NewPathToObjectDataCache does not contain it. 
-			PathToObjectDataCache = MoveTemp(NewPathToObjectDataCache);
-
-			// The tree view requires the item source to only contain the root items. Children are discovered via GetObjectRowChildren. We re-use GetObjectRowChildren to remove any non-root nodes.
-			BuildRootObjectRowData();
-			ReplicatedObjects->OnItemsChanged();
-		}
-	}
-
-	void SReplicationStreamViewer::RefreshPropertyData()
-	{
-		PropertySection->RefreshPropertyData();
-	}
-
 	void SReplicationStreamViewer::SelectObjects(TConstArrayView<FSoftObjectPath> Objects)
 	{
 		TArray<TSharedPtr<FReplicatedObjectData>> NewSelectedItems; 
@@ -152,7 +105,36 @@
 			ReplicatedObjects->SetExpandedItems(ItemsToExpand, true);
 		}
 	}
-	
+
+	TArray<TSharedPtr<FReplicatedObjectData>> SReplicationStreamViewer::GetSelectedOutlinerObjects() const
+	{
+		TArray<TSharedPtr<FReplicatedObjectData>> SelectedItems = ReplicatedObjects->GetSelectedItems();
+		// Items may have been removed this tick. However, selected items may not have been updated yet because STreeView processes item changes at the end of tick. 
+		SelectedItems.SetNum(Algo::RemoveIf(SelectedItems, [this](const TSharedPtr<FReplicatedObjectData>& ObjectData)
+		{
+			return !PropertiesModel->ContainsObjects({ ObjectData->GetObjectPath() });
+		}));
+		return SelectedItems;
+	}
+
+	void SReplicationStreamViewer::Tick(const FGeometry& AllottedGeometry, const double InCurrentTime, const float InDeltaTime)
+	{
+		if (bHasRequestedObjectRefresh)
+		{
+			bHasRequestedObjectRefresh = false;
+			bHasRequestedPropertyRefresh = true;
+			RefreshObjectData();
+		}
+
+		if (bHasRequestedPropertyRefresh)
+		{
+			bHasRequestedPropertyRefresh = false;
+			RefreshPropertyData();
+		}
+		
+		IReplicationStreamViewer::Tick(AllottedGeometry, InCurrentTime, InDeltaTime);
+	}
+
 	TSharedRef<FReplicatedObjectData> SReplicationStreamViewer::AllocateObjectData(FSoftObjectPath ObjectPath)
 	{
 		return MakeShared<FReplicatedObjectData>(MoveTemp(ObjectPath));
@@ -206,7 +188,7 @@
 			.OnDeleteItems(InArgs._OnDeleteObjects)
 			.OnSelectionChanged_Lambda([this]()
 			{
-				RefreshPropertyData();
+				RequestPropertyDataRefresh();
 			})
 			.Columns(Columns)
 			.ExpandableColumnLabel(ReplicationColumns::TopLevel::LabelColumnId)
@@ -268,6 +250,53 @@
 				]
 			];
 	}
+	
+	void SReplicationStreamViewer::RefreshObjectData()
+	{
+		// Re-using existing instances is tricky: we cannot update the object path in an item because the list view will no detect this change;
+		// list view only looks at the shared ptr address. So the UI will not be refreshed. Since the number of items will be small, just reallocate... 
+		AllObjectRowData.Empty();
 
+		// Try to re-use old instances by using the old PathToObjectDataCache. This is also done so the expansion states restore correctly in the tree view.
+		TMap<FSoftObjectPath, TSharedPtr<FReplicatedObjectData>> NewPathToObjectDataCache;
+		
+		// Do a complete refresh.
+		// Complete refresh is acceptable because the list is updated infrequently and typically small < 500 items.
+		// An alternative would be to change RefreshObjectData to be called with two variables ObjectsAdded and ObjectsRemoved.
+		PropertiesModel->ForEachReplicatedObject([this, &NewPathToObjectDataCache](const FSoftObjectPath& ObjectPath) mutable
+		{
+			TOptional<IObjectHierarchyModel::FParentInfo> ParentInfo = ObjectHierarchy->GetParentInfo(ObjectPath);
+			const bool bIsActor = !ParentInfo; 
+			if (bIsActor || ShouldDisplayObject(ObjectPath, ParentInfo->Relationship))
+			{
+				const TSharedPtr<FReplicatedObjectData>* ExistingItem = PathToObjectDataCache.Find(ObjectPath);
+				ExistingItem = ExistingItem ? ExistingItem : NewPathToObjectDataCache.Find(ObjectPath);
+				const TSharedRef<FReplicatedObjectData> Item = ExistingItem ? ExistingItem->ToSharedRef() : AllocateObjectData(ObjectPath);
+				AllObjectRowData.AddUnique(Item);
+				NewPathToObjectDataCache.Emplace(ObjectPath, Item);
+				
+				BuildObjectHierarchyIfNeeded(Item, NewPathToObjectDataCache);
+			}
+			
+			return EBreakBehavior::Continue;
+		});
+
+		// Only refresh the tree if it is necessary as it causes us to select stuff in the subobject view
+		if (!PathToObjectDataCache.OrderIndependentCompareEqual(NewPathToObjectDataCache))
+		{
+			// If an item was removed, then NewPathToObjectDataCache does not contain it. 
+			PathToObjectDataCache = MoveTemp(NewPathToObjectDataCache);
+
+			// The tree view requires the item source to only contain the root items. Children are discovered via GetObjectRowChildren. We re-use GetObjectRowChildren to remove any non-root nodes.
+			BuildRootObjectRowData();
+			ReplicatedObjects->OnItemsChanged();
+		}
+	}
+
+	void SReplicationStreamViewer::RefreshPropertyData()
+	{
+		PropertySection->RefreshPropertyData();
+	}
+
 	void SReplicationStreamViewer::BuildRootObjectRowData()
 	{
diff --git a/Engine/Plugins/Developer/Concert/ConcertUI/ConcertSharedSlate/Source/ConcertSharedSlate/Private/Replication/Editor/View/ObjectViewer/SReplicationStreamViewer.h b/Engine/Plugins/Developer/Concert/ConcertUI/ConcertSharedSlate/Source/ConcertSharedSlate/Private/Replication/Editor/View/ObjectViewer/SReplicationStreamViewer.h
--- a/Engine/Plugins/Developer/Concert/ConcertUI/ConcertSharedSlate/Source/ConcertSharedSlate/Private/Replication/Editor/View/ObjectViewer/SReplicationStreamViewer.h
+++ b/Engine/Plugins/Developer/Concert/ConcertUI/ConcertSharedSlate/Source/ConcertSharedSlate/Private/Replication/Editor/View/ObjectViewer/SReplicationStreamViewer.h
@@ -97,8 +97,8 @@ namespace UE::ConcertSharedSlate
 		virtual TArray<FSoftObjectPath> GetObjectsBeingPropertyEdited() const override;
 		//~ End IReplicationStreamViewer Interface
 
-		void RefreshObjectData();
-		void RefreshPropertyData();
+		void RequestObjectDataRefresh() { bHasRequestedObjectRefresh = true; }
+		void RequestPropertyDataRefresh() { bHasRequestedPropertyRefresh = true; }
 
 		/** Selects the given objects. */
 		void SelectObjects(TConstArrayView<FSoftObjectPath> Objects);
@@ -106,7 +106,11 @@ namespace UE::ConcertSharedSlate
 		void ExpandObjects(TConstArrayView<FSoftObjectPath> Objects, bool bRecursive);
 
 		/** @return Gets the root objects selected in the outliner; the subobject view chooses which of these objects (or their subobjects) end up in GetSelectedObjectShowingProperties. */
-		TArray<TSharedPtr<FReplicatedObjectData>> GetSelectedOutlinerObjects() const { return ReplicatedObjects->GetSelectedItems(); }
+		TArray<TSharedPtr<FReplicatedObjectData>> GetSelectedOutlinerObjects() const;
+
+		//~ Begin SWidget Interface
+		virtual void Tick(const FGeometry& AllottedGeometry, const double InCurrentTime, const float InDeltaTime) override;
+		//~ End SWidget Interface
 		
 	private:
 
@@ -136,13 +140,19 @@ namespace UE::ConcertSharedSlate
 		/** View options for the object outliner. */
 		FStreamViewerObjectViewOptions ObjectViewOptions;
 
+		bool bHasRequestedObjectRefresh = false;
+		bool bHasRequestedPropertyRefresh = false;
+
 		static TSharedRef<FReplicatedObjectData> AllocateObjectData(FSoftObjectPath ObjectPath);
 
 		// Widget creation helpers
 		TSharedRef<SWidget> CreateContentWidget(const FArguments& InArgs);
 		TSharedRef<SWidget> CreateOutlinerSection(const FArguments& InArgs);
 		TSharedRef<SWidget> CreatePropertiesSection(const FArguments& InArgs);
-
+		
+		void RefreshObjectData();
+		void RefreshPropertyData();
+		
 		/** Sets RootObjectRowData to all non-root nodes from ObjectRowData. */
 		void BuildRootObjectRowData();
 
@@ -157,8 +167,8 @@ namespace UE::ConcertSharedSlate
 		/** Called in response to subobject display view option being changed. Rebuilds the entire hierarchy. */
 		void OnSubobjectViewOptionToggled()
 		{
-			RefreshObjectData();
-			RefreshPropertyData();
+			RequestObjectDataRefresh();
+			RequestPropertyDataRefresh();
 		}
 		bool ShouldDisplayObject(const FSoftObjectPath& Object, EChildRelationship Relationship) const;
 	};
diff --git a/Engine/Plugins/Developer/Concert/ConcertUI/ConcertSharedSlate/Source/ConcertSharedSlate/Public/Replication/Editor/Model/IReplicationStreamModel.h b/Engine/Plugins/Developer/Concert/ConcertUI/ConcertSharedSlate/Source/ConcertSharedSlate/Public/Replication/Editor/Model/IReplicationStreamModel.h
--- a/Engine/Plugins/Developer/Concert/ConcertUI/ConcertSharedSlate/Source/ConcertSharedSlate/Public/Replication/Editor/Model/IReplicationStreamModel.h
+++ b/Engine/Plugins/Developer/Concert/ConcertUI/ConcertSharedSlate/Source/ConcertSharedSlate/Public/Replication/Editor/Model/IReplicationStreamModel.h
@@ -10,7 +10,11 @@
 namespace UE::ConcertSharedSlate
 {
 	/**
-	 * Abstracts the concept of mapping objects to properties. This is a read-only interface.
+	 * Abstracts the concept of mapping objects to properties.
+	 * This is a read-only interface.
+	 *
+	 * Objects are allowed to be bound to no properties.
+	 * 
 	 * @see IEditableReplicationStreamModel 
 	 */
 	class CONCERTSHAREDSLATE_API IReplicationStreamModel : public TSharedFromThis<IReplicationStreamModel>
