Fix Tree view children are not found by search

diff --git a/Engine/Plugins/Developer/Concert/ConcertUI/ConcertSharedSlate/Source/ConcertSharedSlate/Private/Replication/Editor/View/Tree/SReplicationTreeView.h b/Engine/Plugins/Developer/Concert/ConcertUI/ConcertSharedSlate/Source/ConcertSharedSlate/Private/Replication/Editor/View/Tree/SReplicationTreeView.h
--- a/Engine/Plugins/Developer/Concert/ConcertUI/ConcertSharedSlate/Source/ConcertSharedSlate/Private/Replication/Editor/View/Tree/SReplicationTreeView.h
+++ b/Engine/Plugins/Developer/Concert/ConcertUI/ConcertSharedSlate/Source/ConcertSharedSlate/Private/Replication/Editor/View/Tree/SReplicationTreeView.h
@@ -293,7 +293,10 @@ namespace UE::ConcertSharedSlate
 		FName ExpandableColumnId;
 
 		TArray<TSharedPtr<TItemType>>* AllRootItems = nullptr;
+		/** Contains only the root items that passed the filters */
 		TArray<TSharedPtr<TItemType>> FilteredRootItems;
+		/** Includes ALL items in the hierarchy that have passed the filter. */
+		TSet<TSharedPtr<TItemType>> AllFilteredItems;
 
 		bool bFilterChanged = false;
 		bool bRequestedSort = false;
@@ -331,8 +334,15 @@ namespace UE::ConcertSharedSlate
 		void OnSearchTextChanged(const FText& InSearchText);
 
 		// Filtering
+		enum class EFilterResult : uint8
+		{
+			ItemOrChildrenPassFilter,
+			NoneInHierarchyPassFilter
+		};
+		
 		void PopulateSearchStrings(const TSharedPtr<TItemType>& Item, TArray<FString>& OutSearchStrings);
 		void ReapplyFilters();
+		EFilterResult ApplyFiltersRecursive(const TSharedPtr<TItemType>& Item, TSet<TSharedPtr<TItemType>>& FilteredItemsToShow);
 		bool PassesFilters(const TSharedPtr<TItemType>& Item);
 
 		// Sorting
@@ -467,7 +477,12 @@ namespace UE::ConcertSharedSlate
 		{
 			OnGetChildrenDelegate.Execute(Item, [this, &OutChildren](TSharedPtr<TItemType> ItemToAdd)
 			{
-				if (PassesFilters(ItemToAdd))
+				const bool bAllowedByFilter =
+					// When we applied the filter, was this item or one of its children allowed?
+					AllFilteredItems.Contains(ItemToAdd)
+					// Handle case where no filter has been applied, yet
+					|| PassesFilters(ItemToAdd);
+				if (bAllowedByFilter)
 				{
 					OutChildren.Add(ItemToAdd);
 				}
@@ -573,11 +588,12 @@ namespace UE::ConcertSharedSlate
 
 		// Reset the list of displayed activities.
 		FilteredRootItems.Reset(AllRootItems->Num());
+		AllFilteredItems.Reset();
 
 		// Apply the filter.
 		for (const TSharedPtr<TItemType>& Activity : *AllRootItems)
 		{
-			if (PassesFilters(Activity))
+			if (ApplyFiltersRecursive(Activity, AllFilteredItems) == EFilterResult::ItemOrChildrenPassFilter)
 			{
 				FilteredRootItems.Add(Activity);
 			}
@@ -596,6 +612,27 @@ namespace UE::ConcertSharedSlate
 		TreeView->RequestListRefresh();
 	}
 
+	template <typename TItemType>
+	typename SReplicationTreeView<TItemType>::EFilterResult SReplicationTreeView<TItemType>::ApplyFiltersRecursive(const TSharedPtr<TItemType>& Item, TSet<TSharedPtr<TItemType>>& FilteredItemsToShow)
+	{
+		bool bPassesAtLeastOnce = false;
+		if (PassesFilters(Item))
+		{
+			bPassesAtLeastOnce = true;
+			FilteredItemsToShow.Add(Item);
+		}
+
+		if (OnGetChildrenDelegate.IsBound())
+		{
+			OnGetChildrenDelegate.Execute(Item, [this, &FilteredItemsToShow, &bPassesAtLeastOnce](TSharedPtr<TItemType> ItemToAdd)
+			{
+				bPassesAtLeastOnce |= ApplyFiltersRecursive(ItemToAdd, FilteredItemsToShow) == EFilterResult::ItemOrChildrenPassFilter;
+			});
+		}
+
+		return bPassesAtLeastOnce ? EFilterResult::ItemOrChildrenPassFilter : EFilterResult::NoneInHierarchyPassFilter;
+	}
+
 	template <typename TItemType>
 	bool SReplicationTreeView<TItemType>::PassesFilters(const TSharedPtr<TItemType>& Item)
 	{
