Refactor property tree UI in preparation for extracting property filtering to ConcertSharedClientSlate

diff --git a/Engine/Plugins/Developer/Concert/ConcertUI/ConcertSharedSlate/Source/ConcertSharedSlate/Private/Replication/Editor/View/ObjectViewer/Property/SFilteredPropertyTreeView.cpp b/Engine/Plugins/Developer/Concert/ConcertUI/ConcertSharedSlate/Source/ConcertSharedSlate/Private/Replication/Editor/View/ObjectViewer/Property/SFilteredPropertyTreeView.cpp
new file mode 100644
--- /dev/null
+++ b/Engine/Plugins/Developer/Concert/ConcertUI/ConcertSharedSlate/Source/ConcertSharedSlate/Private/Replication/Editor/View/ObjectViewer/Property/SFilteredPropertyTreeView.cpp
@@ -0,0 +1,172 @@
+// Copyright Epic Games, Inc. All Rights Reserved.
+
+#include "SFilteredPropertyTreeView.h"
+
+#include "Filters/SBasicFilterBar.h"
+#include "PropertyFilter_ByPropertyType.h"
+#include "PropertyFrontendFilter.h"
+#include "SPropertyTreeView.h"
+#include "Replication/Editor/Model/ReplicatedPropertyData.h"
+
+#include "UObject/UnrealType.h"
+
+#define LOCTEXT_NAMESPACE "SReplicatedPropertiesView"
+
+namespace UE::ConcertSharedSlate
+{
+	/** Exposes SetFrontendFilterActive so we can manually enable the default filters */
+	class SReplicationFilterBar : public SBasicFilterBar<TSharedPtr<FReplicatedPropertyData>>
+	{
+		using Super = SBasicFilterBar<TSharedPtr<FReplicatedPropertyData>>;
+	public:
+
+		SLATE_BEGIN_ARGS(SReplicationFilterBar)
+		{}
+			SLATE_EVENT(FOnFilterChanged, OnFilterChanged)
+		SLATE_END_ARGS()
+
+		void Construct(const FArguments& InArgs, TArray<SFilteredPropertyTreeView::FFilterRef> AllFilters)
+		{
+			Super::Construct(
+			Super::FArguments()
+				.FilterPillStyle(EFilterPillStyle::Basic)
+				.CustomFilters(MoveTemp(AllFilters))
+				.OnFilterChanged(InArgs._OnFilterChanged)
+				.UseSectionsForCategories(true)
+				);
+		}
+
+		// Expose from SBasicFilterBar
+		using Super::SetFrontendFilterActive;
+
+		void SetFilterVisuallyEnabled(const SFilteredPropertyTreeView::FFilterRef& Filter, bool bEnabled)
+		{
+			const TSharedRef<SFilter>* FilterWidget = Filters.FindByPredicate([&Filter](const TSharedRef<SFilter>& FilterWidget)
+			{
+				return FilterWidget->GetFrontendFilter() == Filter;
+			});
+			if (ensure(FilterWidget))
+			{
+				FilterWidget->Get().SetEnabled(bEnabled);
+			}
+		}
+	};
+	
+	void SFilteredPropertyTreeView::Construct(const FArguments& InArgs)
+	{
+		const FBuildFilterBarResult Filters = BuildFilterBar();
+		ChildSlot
+		[
+			SAssignNew(ReplicatedProperties, SPropertyTreeView)
+				.Columns(InArgs._Columns)
+				.ExpandableColumnLabel(InArgs._ExpandableColumnLabel)
+				.PrimarySort(InArgs._PrimarySort)
+				.SecondarySort(InArgs._SecondarySort)
+				.SelectionMode(InArgs._SelectionMode)
+				.FilterItem(this, &SFilteredPropertyTreeView::PassesFilters)
+				.LeftOfSearchBar()
+				[
+					SNew(SHorizontalBox)
+
+					// The combo button for selecting the property filters
+					+SHorizontalBox::Slot()
+					.AutoWidth()
+					[
+						SBasicFilterBar<TSharedPtr<FReplicatedPropertyData>>::MakeAddFilterButton(FilterBar.ToSharedRef())
+					]
+
+					+SHorizontalBox::Slot()
+					.AutoWidth()
+					[
+						InArgs._LeftOfSearchBar.Widget
+					]
+				]
+				.RightOfSearchBar() [ InArgs._RightOfSearchBar.Widget ]
+				.RowBelowSearchBar() [ FilterBar.ToSharedRef() ]
+				.NoItemsContent() [ SNew(STextBlock).Text(LOCTEXT("AllFitlered", "All properties filtered.")) ]
+		];
+
+		// For better UX, hide certain properties by default (e.g. why would you want to replicate bools?)
+		// Do this AFTER initializing ReplicatedProperties because it triggers the OnItemsChanged callback.
+		for (const FFilterRef& Filter : Filters.DisabledByDefault)
+		{
+			FilterBar->SetFilterCheckState(Filter, ECheckBoxState::Unchecked);
+		}
+		// Show all the other filters as enabled (not greyed out: blue) - they will not run their logic since they are inverse.
+		for (const FFilterRef& Filter : Filters.EnabledByDefault)
+		{
+			// Makes it appear on the bar
+			FilterBar->SetFilterCheckState(Filter, ECheckBoxState::Checked);
+			// Functionally makes it affect the search
+			FilterBar->SetFrontendFilterActive(Filter, true);
+			// Visually makes the button blue (i.e. so it looks enabled)
+			FilterBar->SetFilterVisuallyEnabled(Filter, true);
+		}
+	}
+
+	SFilteredPropertyTreeView::FBuildFilterBarResult SFilteredPropertyTreeView::BuildFilterBar()
+	{
+		TSharedRef<FFilterCategory> CommonCategory = MakeShared<FFilterCategory>(
+			LOCTEXT("CommonCategory.Name", "Common"),
+			LOCTEXT("CommonCategory.ToolTip", "Include commonly replicated properties.")
+		);
+		TSharedRef<FFilterCategory> UncommonCategory = MakeShared<FFilterCategory>(
+			LOCTEXT("UncommonCategory.Advanced", "Uncommon"),
+			LOCTEXT("UncommonCategory.ToolTip", "Include uncommonly replicated properties.")
+		);
+
+		FBuildFilterBarResult Result;
+		using FFrontendFilter = TPropertyFrontendFilter<FPropertyFilter_ByPropertyType>;
+		// Do not show under search bar
+		Result.DisabledByDefault =
+		{
+			// Ordering matters for the drop-down menu next to settings
+			// Common
+			MakeShared<FFrontendFilter>(CommonCategory, LOCTEXT("Ints", "Integer"), LOCTEXT("Ints.Tooltip", "Includes: uint16, uint32, uint64, int16, int32, int64"), TSet<FFieldClass*>{ FUInt16Property::StaticClass(), FUInt32Property::StaticClass(), FUInt64Property::StaticClass(), FInt16Property::StaticClass(), FIntProperty::StaticClass(), FInt64Property::StaticClass() }),
+			MakeShared<FFrontendFilter>(CommonCategory, LOCTEXT("Floats", "Float"), LOCTEXT("Floats.Tooltip", "Includes: float, double"), TSet<FFieldClass*>{ FFloatProperty::StaticClass(), FDoubleProperty::StaticClass() }),
+			MakeShared<FFrontendFilter>(CommonCategory, LOCTEXT("Struct", "Struct"), TSet<FFieldClass*>{ FStructProperty::StaticClass() }),
+			MakeShared<FFrontendFilter>(CommonCategory, LOCTEXT("Containers", "Containers"), LOCTEXT("Containers.Tooltip", "Includes: array, set, map"), TSet<FFieldClass*>{ FArrayProperty::StaticClass(), FSetProperty::StaticClass(), FMapProperty::StaticClass() }),
+
+			// Uncommon
+			MakeShared<FFrontendFilter>(UncommonCategory, LOCTEXT("Bool", "Boolean"), TSet<FFieldClass*>{ FBoolProperty::StaticClass() }),
+			MakeShared<FFrontendFilter>(UncommonCategory, LOCTEXT("Enum", "Enum"), TSet<FFieldClass*>{ FEnumProperty::StaticClass(), FByteProperty::StaticClass() }),
+			MakeShared<FFrontendFilter>(UncommonCategory, LOCTEXT("Text", "Text"), LOCTEXT("Text.Tooltip", "Includes: FName, FString, FText"), TSet<FFieldClass*>{ FNameProperty::StaticClass(), FStrProperty::StaticClass(), FTextProperty::StaticClass() }),
+			MakeShared<FFrontendFilter>(UncommonCategory, LOCTEXT("SoftPtr", "Soft Ptr"), TSet<FFieldClass*>{ FSoftObjectProperty::StaticClass() }),
+		};
+		// Show up under search bar as enabled
+		Result.EnabledByDefault =
+		{
+			// We'll not enable any filters by default because that's the default for other places in the engine, like the Content Browser, too
+		};
+		// Show up in menu
+		TArray<FFilterRef> AllFilters;
+		AllFilters.Append(Result.EnabledByDefault);
+		AllFilters.Append(Result.DisabledByDefault);
+		
+		FilterBar = SNew(SReplicationFilterBar, MoveTemp(AllFilters))
+			.OnFilterChanged(this, &SFilteredPropertyTreeView::RequestRefilter);
+		
+		return Result;
+	}
+
+	bool SFilteredPropertyTreeView::PassesFilters(const TSharedPtr<FReplicatedPropertyData>& ReplicatedPropertyData) const
+	{
+		return FilterBar->GetAllActiveFilters()->Num() == 0 // Return all items when none enabled
+			|| PassesAnyFilters(ReplicatedPropertyData);
+	}
+
+	bool SFilteredPropertyTreeView::PassesAnyFilters(const TSharedPtr<FReplicatedPropertyData>& ReplicatedPropertyData) const
+	{
+		TSharedPtr<TFilterCollection<TSharedPtr<FReplicatedPropertyData>>> FilterCollection = FilterBar->GetAllActiveFilters();
+		for (int32 Index = 0; Index < FilterCollection->Num(); Index++)
+		{
+			if (FilterCollection->GetFilterAtIndex(Index)->PassesFilter(ReplicatedPropertyData))
+			{
+				return true;
+			}
+		}
+		return false;
+	}
+}
+
+#undef LOCTEXT_NAMESPACE
\ No newline at end of file
diff --git a/Engine/Plugins/Developer/Concert/ConcertUI/ConcertSharedSlate/Source/ConcertSharedSlate/Private/Replication/Editor/View/ObjectViewer/Property/SFilteredPropertyTreeView.h b/Engine/Plugins/Developer/Concert/ConcertUI/ConcertSharedSlate/Source/ConcertSharedSlate/Private/Replication/Editor/View/ObjectViewer/Property/SFilteredPropertyTreeView.h
new file mode 100644
--- /dev/null
+++ b/Engine/Plugins/Developer/Concert/ConcertUI/ConcertSharedSlate/Source/ConcertSharedSlate/Private/Replication/Editor/View/ObjectViewer/Property/SFilteredPropertyTreeView.h
@@ -0,0 +1,80 @@
+// Copyright Epic Games, Inc. All Rights Reserved.
+
+#pragma once
+
+#include "SPropertyTreeView.h"
+#include "Replication/Editor/View/Tree/SReplicationTreeView.h"
+
+#include "Filters/SBasicFilterBar.h"
+#include "Widgets/DeclarativeSyntaxSupport.h"
+#include "Widgets/SCompoundWidget.h"
+
+namespace UE::ConcertSharedSlate
+{
+	class FReplicatedPropertyData;
+	class SReplicationFilterBar;
+	
+	/** This widget extends SPropertyTreeView with filtering functionality. */
+	class SFilteredPropertyTreeView : public SCompoundWidget
+	{
+	public:
+		
+		using FFilterRef = TSharedRef<FFilterBase<TSharedPtr<FReplicatedPropertyData>>>;
+		
+		SLATE_BEGIN_ARGS(SFilteredPropertyTreeView)
+		{}
+			/*************** Arguments inherited by SReplicationTreeView ***************/
+		
+			/** The columns this list should have */
+			SLATE_ARGUMENT(TArray<TReplicationColumn<FReplicatedPropertyData>>, Columns)
+			/** The name of the column that will have the SExpanderArrow for the tree view. */
+			SLATE_ARGUMENT(FName, ExpandableColumnLabel)
+			/** Initial primary sort to set. */
+			SLATE_ARGUMENT(FColumnSortInfo, PrimarySort)
+			/** Initial secondary sort to set. */
+			SLATE_ARGUMENT(FColumnSortInfo, SecondarySort)
+		
+			/** How many items are to allowed to be selected */
+			SLATE_ARGUMENT(ESelectionMode::Type, SelectionMode)
+		
+			/** Optional widget to add to the left of the search bar. */
+			SLATE_NAMED_SLOT(FArguments, LeftOfSearchBar)
+			/** Optional widget to add to the left of the search bar. */
+			SLATE_NAMED_SLOT(FArguments, RightOfSearchBar)
+		
+		SLATE_END_ARGS()
+
+		void Construct(const FArguments& InArgs);
+		
+		/** Rebuilds all property data from the property source */
+		void RefreshPropertyData(const TSet<FConcertPropertyChain>& PropertiesToDisplay, const FSoftClassPath& Class, bool bCanReusePropertyData) const
+		{
+			ReplicatedProperties->RefreshPropertyData(PropertiesToDisplay, Class, bCanReusePropertyData);
+		}
+		
+		/** Called when the items need to be refiltered due to the item source changing. */
+		void RequestRefilter() const { ReplicatedProperties->RequestRefilter(); }
+		/** Requests that the given column be resorted, if it currently affects the row sorting. */
+		void RequestResortForColumn(const FName& ColumnId) const { ReplicatedProperties->RequestResortForColumn(ColumnId); }
+
+	private:
+
+		/** The tree view displaying the replicated properties */
+		TSharedPtr<SPropertyTreeView> ReplicatedProperties;
+		/** Displays the active filters*/
+		TSharedPtr<SReplicationFilterBar> FilterBar;
+
+		struct FBuildFilterBarResult
+		{
+			TArray<FFilterRef> EnabledByDefault;
+			TArray<FFilterRef> DisabledByDefault;
+		};
+		/** Build the filter bar and returns the filters that should be active by default. */
+		FBuildFilterBarResult BuildFilterBar();
+
+		/** Runs all filters through this item */
+		bool PassesFilters(const TSharedPtr<FReplicatedPropertyData>& ReplicatedPropertyData) const;
+		bool PassesAnyFilters(const TSharedPtr<FReplicatedPropertyData>& ReplicatedPropertyData) const;
+	};
+}
+
diff --git a/Engine/Plugins/Developer/Concert/ConcertUI/ConcertSharedSlate/Source/ConcertSharedSlate/Private/Replication/Editor/View/ObjectViewer/Property/SPropertyTreeView.cpp b/Engine/Plugins/Developer/Concert/ConcertUI/ConcertSharedSlate/Source/ConcertSharedSlate/Private/Replication/Editor/View/ObjectViewer/Property/SPropertyTreeView.cpp
--- a/Engine/Plugins/Developer/Concert/ConcertUI/ConcertSharedSlate/Source/ConcertSharedSlate/Private/Replication/Editor/View/ObjectViewer/Property/SPropertyTreeView.cpp
+++ b/Engine/Plugins/Developer/Concert/ConcertUI/ConcertSharedSlate/Source/ConcertSharedSlate/Private/Replication/Editor/View/ObjectViewer/Property/SPropertyTreeView.cpp
@@ -2,206 +2,94 @@
 
 #include "SPropertyTreeView.h"
 
-#include "Filters/SBasicFilterBar.h"
-#include "PropertyFilter_ByPropertyType.h"
-#include "PropertyFrontendFilter.h"
 #include "Replication/Editor/Model/ReplicatedPropertyData.h"
-#include "Replication/Editor/View/DisplayUtils.h"
 
+#include "Algo/ForEach.h"
 #include "UObject/UnrealType.h"
 
 #define LOCTEXT_NAMESPACE "SReplicatedPropertiesView"
 
 namespace UE::ConcertSharedSlate
 {
-	/** Exposes SetFrontendFilterActive so we can manually enable the default filters */
-	class SReplicationFilterBar : public SBasicFilterBar<TSharedPtr<FReplicatedPropertyData>>
-	{
-		using Super = SBasicFilterBar<TSharedPtr<FReplicatedPropertyData>>;
-	public:
-
-		SLATE_BEGIN_ARGS(SReplicationFilterBar)
-		{}
-			SLATE_EVENT(FOnFilterChanged, OnFilterChanged)
-		SLATE_END_ARGS()
-
-		void Construct(const FArguments& InArgs, TArray<SPropertyTreeView::FFilterRef> AllFilters)
-		{
-			Super::Construct(
-			Super::FArguments()
-				.FilterPillStyle(EFilterPillStyle::Basic)
-				.CustomFilters(MoveTemp(AllFilters))
-				.OnFilterChanged(InArgs._OnFilterChanged)
-				.UseSectionsForCategories(true)
-				);
-		}
-
-		// Expose from SBasicFilterBar
-		using Super::SetFrontendFilterActive;
-
-		void SetFilterVisuallyEnabled(const SPropertyTreeView::FFilterRef& Filter, bool bEnabled)
-		{
-			const TSharedRef<SFilter>* FilterWidget = Filters.FindByPredicate([&Filter](const TSharedRef<SFilter>& FilterWidget)
-			{
-				return FilterWidget->GetFrontendFilter() == Filter;
-			});
-			if (ensure(FilterWidget))
-			{
-				FilterWidget->Get().SetEnabled(bEnabled);
-			}
-		}
-	};
-	
 	void SPropertyTreeView::Construct(const FArguments& InArgs)
 	{
-		SelectedObjectsAttribute = InArgs._SelectedObjects;
-		NameModel = InArgs._NameModel;
-		
-		const FBuildFilterBarResult Filters = BuildFilterBar();
 		ChildSlot
 		[
-			SAssignNew(ReplicatedProperties, SReplicationTreeView<FReplicatedPropertyData>)
-				.RootItemsSource(InArgs._RootItemsSource)
-				.OnGetChildren(InArgs._OnGetChildren)
+			SAssignNew(TreeView, SReplicationTreeView<FReplicatedPropertyData>)
+				.RootItemsSource(&RootPropertyRowData)
+				.OnGetChildren(this, &SPropertyTreeView::GetPropertyRowChildren)
+				.FilterItem(InArgs._FilterItem)
 				.Columns(InArgs._Columns)
 				.ExpandableColumnLabel(InArgs._ExpandableColumnLabel)
 				.PrimarySort(InArgs._PrimarySort)
 				.SecondarySort(InArgs._SecondarySort)
 				.SelectionMode(InArgs._SelectionMode)
-				.FilterItem(this, &SPropertyTreeView::PassesFilters)
-				.LeftOfSearchBar()
-				[
-					SNew(SHorizontalBox)
-
-					// The combo button for selecting the property filters
-					+SHorizontalBox::Slot()
-					.AutoWidth()
-					[
-						SBasicFilterBar<TSharedPtr<FReplicatedPropertyData>>::MakeAddFilterButton(FilterBar.ToSharedRef())
-					]
-
-					+SHorizontalBox::Slot()
-					.AutoWidth()
-					[
-						InArgs._LeftOfSearchBar.Widget
-					]
-				]
+				.LeftOfSearchBar() [ InArgs._LeftOfSearchBar.Widget ]
 				.RightOfSearchBar() [ InArgs._RightOfSearchBar.Widget ]
-				.RowBelowSearchBar() [ FilterBar.ToSharedRef() ]
-				.NoItemsContent() [ SNew(STextBlock).Text(this, &SPropertyTreeView::GetAllFilteredText) ]
+				.RowBelowSearchBar() [ InArgs._RowBelowSearchBar.Widget ]
+				.NoItemsContent() [ InArgs._NoItemsContent.Widget ]
 		];
+	}
 
-		// For better UX, hide certain properties by default (e.g. why would you want to replicate bools?)
-		// Do this AFTER initializing ReplicatedProperties because it triggers the OnItemsChanged callback.
-		for (const FFilterRef& Filter : Filters.DisabledByDefault)
+	void SPropertyTreeView::RefreshPropertyData(const TSet<FConcertPropertyChain>& PropertiesToDisplay, const FSoftClassPath& Class, const bool bCanReuseExistingRowItems)
+	{
+		if (!bCanReuseExistingRowItems)
 		{
-			FilterBar->SetFilterCheckState(Filter, ECheckBoxState::Unchecked);
+			ChainToPropertyDataCache.Reset();
 		}
-		// Show all the other filters as enabled (not greyed out: blue) - they will not run their logic since they are inverse.
-		for (const FFilterRef& Filter : Filters.EnabledByDefault)
+		
+		// Try to re-use old instances by using the old ChainToPropertyDataCache. This is also done so the expansion states restore correctly in the tree view.
+		TMap<FConcertPropertyChain, TSharedPtr<FReplicatedPropertyData>> NewChainToPropertyDataCache;
+		
+		PropertyRowData.Empty();
+		for (const FConcertPropertyChain& PropertyChain : PropertiesToDisplay)
 		{
-			// Makes it appear on the bar
-			FilterBar->SetFilterCheckState(Filter, ECheckBoxState::Checked);
-			// Functionally makes it affect the search
-			FilterBar->SetFrontendFilterActive(Filter, true);
-			// Visually makes the button blue (i.e. so it looks enabled)
-			FilterBar->SetFilterVisuallyEnabled(Filter, true);
+			const TSharedPtr<FReplicatedPropertyData>* ExistingItem = ChainToPropertyDataCache.Find(PropertyChain);
+			const TSharedRef<FReplicatedPropertyData> Item = ExistingItem ? ExistingItem->ToSharedRef() : AllocatePropertyData(Class, PropertyChain);
+			PropertyRowData.Emplace(Item);
+			NewChainToPropertyDataCache.Emplace(PropertyChain, Item);
 		}
-	}
-
-	void SPropertyTreeView::OnItemsChanged() const
-	{
-		ReplicatedProperties->OnItemsChanged();
-	}
-
-	void SPropertyTreeView::RequestResortForColumn(const FName& ColumnId)
-	{
-		ReplicatedProperties->RequestResortForColumn(ColumnId);
-	}
 
-	SPropertyTreeView::FBuildFilterBarResult SPropertyTreeView::BuildFilterBar()
-	{
-		TSharedRef<FFilterCategory> CommonCategory = MakeShared<FFilterCategory>(
-			LOCTEXT("CommonCategory.Name", "Common"),
-			LOCTEXT("CommonCategory.ToolTip", "Include commonly replicated properties.")
-		);
-		TSharedRef<FFilterCategory> UncommonCategory = MakeShared<FFilterCategory>(
-			LOCTEXT("UncommonCategory.Advanced", "Uncommon"),
-			LOCTEXT("UncommonCategory.ToolTip", "Include uncommonly replicated properties.")
-		);
-
-		FBuildFilterBarResult Result;
-		using FFrontendFilter = TPropertyFrontendFilter<FPropertyFilter_ByPropertyType>;
-		// Do not show under search bar
-		Result.DisabledByDefault =
-		{
-			// Ordering matters for the drop-down menu next to settings
-			// Common
-			MakeShared<FFrontendFilter>(MoveTemp(CommonCategory), LOCTEXT("Ints", "Integer"), LOCTEXT("Ints.Tooltip", "Includes: uint16, uint32, uint64, int16, int32, int64"), TSet<FFieldClass*>{ FUInt16Property::StaticClass(), FUInt32Property::StaticClass(), FUInt64Property::StaticClass(), FInt16Property::StaticClass(), FIntProperty::StaticClass(), FInt64Property::StaticClass() }),
-			MakeShared<FFrontendFilter>(MoveTemp(CommonCategory), LOCTEXT("Floats", "Float"), LOCTEXT("Floats.Tooltip", "Includes: float, double"), TSet<FFieldClass*>{ FFloatProperty::StaticClass(), FDoubleProperty::StaticClass() }),
-			MakeShared<FFrontendFilter>(MoveTemp(CommonCategory), LOCTEXT("Struct", "Struct"), TSet<FFieldClass*>{ FStructProperty::StaticClass() }),
-			MakeShared<FFrontendFilter>(MoveTemp(CommonCategory), LOCTEXT("Containers", "Containers"), LOCTEXT("Containers.Tooltip", "Includes: array, set, map"), TSet<FFieldClass*>{ FArrayProperty::StaticClass(), FSetProperty::StaticClass(), FMapProperty::StaticClass() }),
-
-			// Uncommon
-			MakeShared<FFrontendFilter>(MoveTemp(UncommonCategory), LOCTEXT("Bool", "Boolean"), TSet<FFieldClass*>{ FBoolProperty::StaticClass() }),
-			MakeShared<FFrontendFilter>(MoveTemp(UncommonCategory), LOCTEXT("Enum", "Enum"), TSet<FFieldClass*>{ FEnumProperty::StaticClass(), FByteProperty::StaticClass() }),
-			MakeShared<FFrontendFilter>(MoveTemp(UncommonCategory), LOCTEXT("Text", "Text"), LOCTEXT("Text.Tooltip", "Includes: FName, FString, FText"), TSet<FFieldClass*>{ FNameProperty::StaticClass(), FStrProperty::StaticClass(), FTextProperty::StaticClass() }),
-			MakeShared<FFrontendFilter>(MoveTemp(UncommonCategory), LOCTEXT("SoftPtr", "Soft Ptr"), TSet<FFieldClass*>{ FSoftObjectProperty::StaticClass() }),
-		};
-		// Show up under search bar as enabled
-		Result.EnabledByDefault =
-		{
-			// We'll not enable any filters by default because that's the default for other places in the engine, like the Content Browser, too
-		};
-		// Show up in menu
-		TArray<FFilterRef> AllFilters;
-		AllFilters.Append(Result.EnabledByDefault);
-		AllFilters.Append(Result.DisabledByDefault);
-		
-		FilterBar = SNew(SReplicationFilterBar, MoveTemp(AllFilters))
-			.OnFilterChanged(this, &SPropertyTreeView::OnItemsChanged);
+		// If an item was removed, then NewPathToPropertyDataCache does not contain it. 
+		ChainToPropertyDataCache = MoveTemp(NewChainToPropertyDataCache);
 		
-		return Result;
+		// The tree view requires the item source to only contain the root items.
+		BuildRootPropertyRowData();
+
+		TreeView->RequestRefilter();
 	}
 
-	bool SPropertyTreeView::PassesFilters(const TSharedPtr<FReplicatedPropertyData>& ReplicatedPropertyData) const
+	TSharedRef<FReplicatedPropertyData> SPropertyTreeView::AllocatePropertyData(FSoftClassPath OwningClass, FConcertPropertyChain PropertyChain)
 	{
-		return FilterBar->GetAllActiveFilters()->Num() == 0 // Return all items when none enabled
-			|| PassesAnyFilters(ReplicatedPropertyData);
+		return MakeShared<FReplicatedPropertyData>(MoveTemp(OwningClass), MoveTemp(PropertyChain));
 	}
 
-	bool SPropertyTreeView::PassesAnyFilters(const TSharedPtr<FReplicatedPropertyData>& ReplicatedPropertyData) const
+	void SPropertyTreeView::BuildRootPropertyRowData()
 	{
-		TSharedPtr<TFilterCollection<TSharedPtr<FReplicatedPropertyData>>> FilterCollection = FilterBar->GetAllActiveFilters();
-		for (int32 Index = 0; Index < FilterCollection->Num(); Index++)
+		RootPropertyRowData.Empty(PropertyRowData.Num());
+		for (const TSharedPtr<FReplicatedPropertyData>& PropertyData : PropertyRowData)
 		{
-			if (FilterCollection->GetFilterAtIndex(Index)->PassesFilter(ReplicatedPropertyData))
+			if (PropertyData->GetProperty().IsRootProperty())
 			{
-				return true;
+				RootPropertyRowData.Emplace(PropertyData);
 			}
 		}
-		return false;
 	}
 
-	FText SPropertyTreeView::GetAllFilteredText() const
+	void SPropertyTreeView::GetPropertyRowChildren(TSharedPtr<FReplicatedPropertyData> ReplicatedPropertyData, TFunctionRef<void(TSharedPtr<FReplicatedPropertyData>)> ProcessChild)
 	{
-		const TArray<FSoftObjectPath> Objects = SelectedObjectsAttribute.Get();
-		if (Objects.IsEmpty())
+		TArray<TSharedPtr<FReplicatedPropertyData>> Children;
+		
+		// Not the most efficient but it should be fine.
+		for (const TSharedPtr<FReplicatedPropertyData>& Data : PropertyRowData)
 		{
-			return FText::GetEmpty();
+			if (Data->GetProperty().IsDirectChildOf(ReplicatedPropertyData->GetProperty()))
+			{
+				Children.Add(Data);
+			}
 		}
 
-		TSet<FString> Names;
-		Algo::Transform(Objects, Names, [this](const FSoftObjectPath& ObjectPath)
-		{
-			return DisplayUtils::GetObjectDisplayText(ObjectPath, NameModel).ToString();
-		});
-		
-		const FText ObjectsText = FText::FromString(FString::Join(Names, TEXT(", ")));
-		return FText::Format(LOCTEXT("AllFitleredFmt", "All properties filtered for selected {0}|plural(one=object,other=objects): {1}."),
-			Objects.Num(),
-			ObjectsText
-		);
+		Algo::ForEach(Children, [&ProcessChild](const TSharedPtr<FReplicatedPropertyData>& Data){ ProcessChild(Data); });
 	}
 }
 
diff --git a/Engine/Plugins/Developer/Concert/ConcertUI/ConcertSharedSlate/Source/ConcertSharedSlate/Private/Replication/Editor/View/ObjectViewer/Property/SPropertyTreeView.h b/Engine/Plugins/Developer/Concert/ConcertUI/ConcertSharedSlate/Source/ConcertSharedSlate/Private/Replication/Editor/View/ObjectViewer/Property/SPropertyTreeView.h
--- a/Engine/Plugins/Developer/Concert/ConcertUI/ConcertSharedSlate/Source/ConcertSharedSlate/Private/Replication/Editor/View/ObjectViewer/Property/SPropertyTreeView.h
+++ b/Engine/Plugins/Developer/Concert/ConcertUI/ConcertSharedSlate/Source/ConcertSharedSlate/Private/Replication/Editor/View/ObjectViewer/Property/SPropertyTreeView.h
@@ -2,19 +2,21 @@
 
 #pragma once
 
+#include "Replication/Data/ConcertPropertySelection.h"
+#include "Replication/Editor/View/Tree/SReplicationTreeView.h"
+
 #include "Filters/SBasicFilterBar.h"
 #include "Widgets/DeclarativeSyntaxSupport.h"
 #include "Widgets/SCompoundWidget.h"
-#include "Replication/Editor/View/Tree/SReplicationTreeView.h"
 
 namespace UE::ConcertSharedSlate
 {
-	class FReplicatedObjectData;
 	class FReplicatedPropertyData;
-	class IObjectNameModel;
-	class SReplicationFilterBar;
 	
-	/** Displays a searchable tree view of properties (SReplicationTreeView) and decorates it with a SBasicFilterBar. */
+	/**
+	 * This widget knows how to display a list of properties in a tree view.
+	 * It generates the items and exposes extension points for more advanced UI, such as filtering.
+	 */
 	class SPropertyTreeView : public SCompoundWidget
 	{
 	public:
@@ -25,11 +27,8 @@ namespace UE::ConcertSharedSlate
 		{}
 			/*************** Arguments inherited by SReplicationTreeView ***************/
 		
-			/** The items to display */
-			SLATE_ARGUMENT(TArray<TSharedPtr<FReplicatedPropertyData>>*, RootItemsSource)
-		
-			/** Gets an items children for the tree view */
-			SLATE_EVENT(SReplicationTreeView<FReplicatedPropertyData>::FGetItemChildren, OnGetChildren)
+			/** Optional callback to do even more filtering of items. */
+			SLATE_EVENT(SReplicationTreeView<FReplicatedPropertyData>::FCustomFilter, FilterItem)
 		
 			/** The columns this list should have */
 			SLATE_ARGUMENT(TArray<TReplicationColumn<FReplicatedPropertyData>>, Columns)
@@ -43,56 +42,59 @@ namespace UE::ConcertSharedSlate
 			/** How many items are to allowed to be selected */
 			SLATE_ARGUMENT(ESelectionMode::Type, SelectionMode)
 		
-			/** Optional. If set, this determines the display text for objects. */
-			SLATE_ARGUMENT(TSharedPtr<IObjectNameModel>, NameModel)
-		
 			/** Optional widget to add to the left of the search bar. */
 			SLATE_NAMED_SLOT(FArguments, LeftOfSearchBar)
 			/** Optional widget to add to the left of the search bar. */
 			SLATE_NAMED_SLOT(FArguments, RightOfSearchBar)
+		
+			/** Optional widget to add between the search bar and the table view (e.g. a SBasicFilterBar). */
+			SLATE_NAMED_SLOT(FArguments, RowBelowSearchBar)
+		
+			/** Optional, alternate content to show instead of the tree view when there are no rows. */
+			SLATE_NAMED_SLOT(FArguments, NoItemsContent)
 
-			/*************** Own arguments ***************/
-			// Please add new arguments here in the future
-
-			/** Gets the objects being property edited. */
-			SLATE_ATTRIBUTE(TArray<FSoftObjectPath>, SelectedObjects)
 		SLATE_END_ARGS()
 
 		void Construct(const FArguments& InArgs);
 
+		/**
+		 * Rebuilds all property data from the property source.
+		 *
+		 * @param PropertiesToDisplay The properties to display
+		 * @param Class The class from which the PropertiesToDisplay come
+		 * @param bCanReuseExistingRowItems True, will try to reuse rows for properties in the tree already (retains selected rows).
+		 *	Set this to false, if all rows should be regenerated (clears selection).
+		 *	In general, always set this to false if you've changed the object for which you're displaying the class.
+		 */
+		void RefreshPropertyData(const TSet<FConcertPropertyChain>& PropertiesToDisplay, const FSoftClassPath& Class, bool bCanReuseExistingRowItems);
+		
 		/** Called when the items need to be refiltered due to the item source changing. */
-		void OnItemsChanged() const;
-
+		void RequestRefilter() const { TreeView->RequestRefilter(); }
 		/** Requests that the given column be resorted, if it currently affects the row sorting. */
-		void RequestResortForColumn(const FName& ColumnId);
+		void RequestResortForColumn(const FName& ColumnId) { TreeView->RequestResortForColumn(ColumnId); }
 
 	private:
 
 		/** The tree view displaying the replicated properties */
-		TSharedPtr<SReplicationTreeView<FReplicatedPropertyData>> ReplicatedProperties;
-		/** Displays the active filters*/
-		TSharedPtr<SReplicationFilterBar> FilterBar;
-
-		/** Used to tell the user that the selected object have all properties filtered out. */
-		TAttribute<TArray<FSoftObjectPath>> SelectedObjectsAttribute;
+		TSharedPtr<SReplicationTreeView<FReplicatedPropertyData>> TreeView;
 		
-		/** Optional. If set, this determines the display text for objects. */
-		TSharedPtr<IObjectNameModel> NameModel;
-
-		struct FBuildFilterBarResult
-		{
-			TArray<FFilterRef> EnabledByDefault;
-			TArray<FFilterRef> DisabledByDefault;
-		};
-		/** Build the filter bar and returns the filters that should be active by default. */
-		FBuildFilterBarResult BuildFilterBar();
-
-		/** Runs all filters through this item */
-		bool PassesFilters(const TSharedPtr<FReplicatedPropertyData>& ReplicatedPropertyData) const;
-		bool PassesAnyFilters(const TSharedPtr<FReplicatedPropertyData>& ReplicatedPropertyData) const;
+		/**
+		 * These instances can be subclasses of FReplicatedPropertyData, e.g. FReplicatedPropertyData_Editor.
+		 * Their type can be overridden by subclasses.
+		 * They only have the FReplicatedPropertyData type so they can be passed efficiently to SObjectToPropertyView.
+		 * @see GetPropertyData
+		 */
+		TArray<TSharedPtr<FReplicatedPropertyData>> PropertyRowData;
+		/** The instances of ObjectRowData which do not have any parents. This acts as the item source for the tree view. */
+		TArray<TSharedPtr<FReplicatedPropertyData>> RootPropertyRowData;
+		/** Inverse map of PropertyRowData using FReplicatedPropertyData::GetProperty as key. Contains all elements of PropertyRowData. */
+		TMap<FConcertPropertyChain, TSharedPtr<FReplicatedPropertyData>> ChainToPropertyDataCache;
+		
+		TSharedRef<FReplicatedPropertyData> AllocatePropertyData(FSoftClassPath OwningClass, FConcertPropertyChain PropertyChain);
 
-		/** Gets the message to display when all properties are filtered out. */
-		FText GetAllFilteredText() const;
+		/** Inits RootPropertyRowData from PropertyRowData. */
+		void BuildRootPropertyRowData();
+		void GetPropertyRowChildren(TSharedPtr<FReplicatedPropertyData> ReplicatedPropertyData, TFunctionRef<void(TSharedPtr<FReplicatedPropertyData>)> ProcessChild);
 	};
 }
 
diff --git a/Engine/Plugins/Developer/Concert/ConcertUI/ConcertSharedSlate/Source/ConcertSharedSlate/Private/Replication/Editor/View/ObjectViewer/SReplicatedPropertyView.cpp b/Engine/Plugins/Developer/Concert/ConcertUI/ConcertSharedSlate/Source/ConcertSharedSlate/Private/Replication/Editor/View/ObjectViewer/SReplicatedPropertyView.cpp
--- a/Engine/Plugins/Developer/Concert/ConcertUI/ConcertSharedSlate/Source/ConcertSharedSlate/Private/Replication/Editor/View/ObjectViewer/SReplicatedPropertyView.cpp
+++ b/Engine/Plugins/Developer/Concert/ConcertUI/ConcertSharedSlate/Source/ConcertSharedSlate/Private/Replication/Editor/View/ObjectViewer/SReplicatedPropertyView.cpp
@@ -4,12 +4,10 @@
 
 #include "Replication/Editor/Model/IReplicationStreamModel.h"
 #include "Replication/Editor/Model/ReplicatedObjectData.h"
-#include "Replication/Editor/Model/ReplicatedPropertyData.h"
-#include "Replication/Editor/View/ObjectViewer/Property/SPropertyTreeView.h"
 #include "Replication/Editor/View/SelectionViewerColumns.h"
 
 #include "Algo/AllOf.h"
-#include "Algo/ForEach.h"
+#include "Property/SFilteredPropertyTreeView.h"
 #include "Widgets/Layout/SWidgetSwitcher.h"
 #include "Widgets/Text/STextBlock.h"
 
@@ -20,6 +18,7 @@ namespace UE::ConcertSharedSlate
 	void SReplicatedPropertyView::Construct(const FArguments& InArgs, TSharedRef<IReplicationStreamModel> InPropertiesModel)
 	{
 		PropertiesModel = MoveTemp(InPropertiesModel);
+		
 		GetSelectedRootObjectsDelegate = InArgs._GetSelectedRootObjects;
 		check(GetSelectedRootObjectsDelegate.IsBound());
 		
@@ -45,23 +44,13 @@ namespace UE::ConcertSharedSlate
 			SetPropertyContent(EReplicatedPropertyContent::SelectionTooBig);
 			return;
 		}
+		
 		const FSoftClassPath Class = *SharedClass;
-		if (!ensureMsgf(Class.IsValid(), TEXT("This should only trigger if IReplicationSubobjectView returned a object not contained the model & that the UI is view-only. In that case fix your IReplicationSubobjectView or use it in an UI editor.")))
+		if (!ensure(Class.IsValid()))
 		{
 			return;
 		}
-
-		// If the objects have changed, the classes may share properties.
-		// In that case, below we'd reuse the item pointer, which would cause the tree view to re-use the old row widgets.
-		// However, we must regenerate all column widgets since they may be referencing the object the row was originally built for. So they'd display the state of the previous object still!
-		// Example: Assign property combo-box in Multi-User All Clients view displays who has the property assigned.
-		// Note: If the objects did not change, we definitely want to reuse item pointers since otherwise the user row selection is reset.
-		const bool bCanReusePropertyData = PreviousSelectedObjects == SelectedObjects; // This SHOULD be an order independent compare but usually Num == 1, so whatever
-		if (!bCanReusePropertyData)
-		{
-			ChainToPropertyDataCache.Reset();
-		}
-
+		
 		// Build the set of properties that are shared by all of the selected objects
 		TSet<FConcertPropertyChain> SharedProperties = PropertiesModel->GetAllProperties(SelectedObjects[0]);
 		for (int32 i = 1; i < SelectedObjects.Num(); ++i)
@@ -69,32 +58,17 @@ namespace UE::ConcertSharedSlate
 			SharedProperties = PropertiesModel->GetAllProperties(SelectedObjects[i]).Union(SharedProperties);
 		}
 		
-		// Try to re-use old instances by using the old ChainToPropertyDataCache. This is also done so the expansion states restore correctly in the tree view.
-		TMap<FConcertPropertyChain, TSharedPtr<FReplicatedPropertyData>> NewChainToPropertyDataCache;
+		// If the objects have changed, the classes may share properties.
+		// In that case, below we'd reuse the item pointer, which would cause the tree view to re-use the old row widgets.
+		// However, we must regenerate all column widgets since they may be referencing the object the row was originally built for. So they'd display the state of the previous object still!
+		// Example: Assign property combo-box in Multi-User All Clients view displays who has the property assigned.
+		// Note: If the objects did not change, we definitely want to reuse item pointers since otherwise the user row selection is reset.
+		const bool bCanReusePropertyData = PreviousSelectedObjects == SelectedObjects; // This SHOULD be an order independent compare but usually Num == 1, so whatever
 		
-		PropertyRowData.Empty();
-		for (const FConcertPropertyChain& PropertyChain : SharedProperties)
-		{
-			const TSharedPtr<FReplicatedPropertyData>* ExistingItem = ChainToPropertyDataCache.Find(PropertyChain);
-			const TSharedRef<FReplicatedPropertyData> Item = ExistingItem ? ExistingItem->ToSharedRef() : AllocatePropertyData(Class, PropertyChain);
-			PropertyRowData.Emplace(Item);
-			NewChainToPropertyDataCache.Emplace(PropertyChain, Item);
-		}
-
-		// If an item was removed, then NewPathToPropertyDataCache does not contain it. 
-		ChainToPropertyDataCache = MoveTemp(NewChainToPropertyDataCache);
+		ReplicatedProperties->RefreshPropertyData(SharedProperties, Class, bCanReusePropertyData);
 		
-		// The tree view requires the item source to only contain the root items.
-		BuildRootPropertyRowData();
 		SetPropertyContent(EReplicatedPropertyContent::Properties);
-
 		PreviousSelectedObjects = MoveTemp(SelectedObjects);
-		ReplicatedProperties->OnItemsChanged();
-	}
-
-	void SReplicatedPropertyView::RequestResortForColumn(const FName& ColumnId)
-	{
-		ReplicatedProperties->RequestResortForColumn(ColumnId);
 	}
 
 	TArray<FSoftObjectPath> SReplicatedPropertyView::GetObjectsSelectedForPropertyEditing() const
@@ -131,18 +105,14 @@ namespace UE::ConcertSharedSlate
 			// EReplicatedPropertyContent::Properties
 			+SWidgetSwitcher::Slot()
 			[
-				SAssignNew(ReplicatedProperties, SPropertyTreeView)
-				.RootItemsSource(&RootPropertyRowData)
-				.OnGetChildren(this, &SReplicatedPropertyView::GetPropertyRowChildren)
+				SAssignNew(ReplicatedProperties, SFilteredPropertyTreeView)
 				.Columns(Columns)
 				.ExpandableColumnLabel(ReplicationColumns::Property::LabelColumnId)
 				.PrimarySort(PrimarySort)
 				.SecondarySort(SecondarySort)
 				.SelectionMode(ESelectionMode::Multi)
-				.NameModel(InArgs._NameModel)
 				.LeftOfSearchBar() [ InArgs._LeftOfPropertySearchBar.Widget ]
 				.RightOfSearchBar() [ InArgs._RightOfPropertySearchBar.Widget ]
-				.SelectedObjects(this, &SReplicatedPropertyView::GetObjectsSelectedForPropertyEditing)
 			]
 			
 			// EReplicatedPropertyContent::NoSelection
@@ -164,23 +134,6 @@ namespace UE::ConcertSharedSlate
 			];
 	}
 	
-	TSharedRef<FReplicatedPropertyData> SReplicatedPropertyView::AllocatePropertyData(FSoftClassPath OwningClass, FConcertPropertyChain PropertyChain)
-	{
-		return MakeShared<FReplicatedPropertyData>(MoveTemp(OwningClass), MoveTemp(PropertyChain));
-	}
-
-	void SReplicatedPropertyView::BuildRootPropertyRowData()
-	{
-		RootPropertyRowData.Empty(PropertyRowData.Num());
-		for (const TSharedPtr<FReplicatedPropertyData>& PropertyData : PropertyRowData)
-		{
-			if (PropertyData->GetProperty().IsRootProperty())
-			{
-				RootPropertyRowData.Emplace(PropertyData);
-			}
-		}
-	}
-	
 	TOptional<FSoftClassPath> SReplicatedPropertyView::GetClassForPropertiesFromSelection(const TArray<FSoftObjectPath>& Objects) const
 	{
 		const FSoftClassPath Class = PropertiesModel->GetObjectClass(Objects[0]);
@@ -191,22 +144,6 @@ namespace UE::ConcertSharedSlate
 		return bAllHaveSameClass ? Class : TOptional<FSoftClassPath>{};
 	}
 
-	void SReplicatedPropertyView::GetPropertyRowChildren(TSharedPtr<FReplicatedPropertyData> ReplicatedPropertyData, TFunctionRef<void(TSharedPtr<FReplicatedPropertyData>)> ProcessChild)
-	{
-		TArray<TSharedPtr<FReplicatedPropertyData>> Children;
-		
-		// Not the most efficient but it should be fine.
-		for (const TSharedPtr<FReplicatedPropertyData>& Data : PropertyRowData)
-		{
-			if (Data->GetProperty().IsDirectChildOf(ReplicatedPropertyData->GetProperty()))
-			{
-				Children.Add(Data);
-			}
-		}
-
-		Algo::ForEach(Children, [&ProcessChild](const TSharedPtr<FReplicatedPropertyData>& Data){ ProcessChild(Data); });
-	}
-
 	void SReplicatedPropertyView::SetPropertyContent(EReplicatedPropertyContent Content) const
 	{
 		PropertyContent->SetActiveWidgetIndex(static_cast<int32>(Content));
diff --git a/Engine/Plugins/Developer/Concert/ConcertUI/ConcertSharedSlate/Source/ConcertSharedSlate/Private/Replication/Editor/View/ObjectViewer/SReplicatedPropertyView.h b/Engine/Plugins/Developer/Concert/ConcertUI/ConcertSharedSlate/Source/ConcertSharedSlate/Private/Replication/Editor/View/ObjectViewer/SReplicatedPropertyView.h
--- a/Engine/Plugins/Developer/Concert/ConcertUI/ConcertSharedSlate/Source/ConcertSharedSlate/Private/Replication/Editor/View/ObjectViewer/SReplicatedPropertyView.h
+++ b/Engine/Plugins/Developer/Concert/ConcertUI/ConcertSharedSlate/Source/ConcertSharedSlate/Private/Replication/Editor/View/ObjectViewer/SReplicatedPropertyView.h
@@ -6,6 +6,7 @@
 #include "Replication/Editor/View/SelectionViewerColumns.h"
 
 #include "Misc/Optional.h"
+#include "Property/SFilteredPropertyTreeView.h"
 #include "Widgets/DeclarativeSyntaxSupport.h"
 #include "Widgets/SCompoundWidget.h"
 
@@ -13,9 +14,12 @@ class SWidgetSwitcher;
 
 namespace UE::ConcertSharedSlate
 {
-	class SPropertyTreeView;
+	class SFilteredPropertyTreeView;
 
-	/** Displays the SPropertyTreeView and decorates it with messages that prompt the user for action, e.g. to select an object to view properties. */
+	/**
+	 * Determines which properties are to be displayed based on an IReplicationStreamModel.
+	 * Uses a property tree for displaying. If no properties are displayed, this widget displays a message instead, e.g. to select an object.
+	 */
 	class SReplicatedPropertyView : public SCompoundWidget
 	{
 	public:
@@ -47,19 +51,18 @@ namespace UE::ConcertSharedSlate
 		
 		void RefreshPropertyData();
 		/** Requests that the given column be resorted, if it currently affects the row sorting. */
-		void RequestResortForColumn(const FName& ColumnId);
+		void RequestResortForColumn(const FName& ColumnId) const { ReplicatedProperties->RequestResortForColumn(ColumnId); }
 		
-		const TArray<TSharedPtr<FReplicatedPropertyData>>& GetPropertyRowData() const { return PropertyRowData; }
 		TArray<FSoftObjectPath> GetObjectsSelectedForPropertyEditing() const;
 
 	private:
 		
+		/** Tree view for replicated properties. Content depends on the current object selected. */
+		TSharedPtr<SFilteredPropertyTreeView> ReplicatedProperties;
+		
 		/** The model this view is visualizing. */
 		TSharedPtr<IReplicationStreamModel> PropertiesModel;
 		
-		/** Tree view for replicated properties. Content depends on the current object selected. */
-		TSharedPtr<SPropertyTreeView> ReplicatedProperties;
-		
 		enum class EReplicatedPropertyContent
 		{
 			/** Shows the properties */
@@ -74,33 +77,15 @@ namespace UE::ConcertSharedSlate
 		
 		/** Gets the root objects selected in the object outliner. */
 		FGetSelectedRootObjects GetSelectedRootObjectsDelegate;
-		
-		/**
-		 * These instances can be subclasses of FReplicatedPropertyData, e.g. FReplicatedPropertyData_Editor.
-		 * Their type can be overridden by subclasses.
-		 * They only have the FReplicatedPropertyData type so they can be passed efficiently to SObjectToPropertyView.
-		 * @see GetPropertyData
-		 */
-		TArray<TSharedPtr<FReplicatedPropertyData>> PropertyRowData;
-		/** The instances of ObjectRowData which do not have any parents. This acts as the item source for the tree view. */
-		TArray<TSharedPtr<FReplicatedPropertyData>> RootPropertyRowData;
-		/** Inverse map of PropertyRowData using FReplicatedPropertyData::GetProperty as key. Contains all elements of PropertyRowData. */
-		TMap<FConcertPropertyChain, TSharedPtr<FReplicatedPropertyData>> ChainToPropertyDataCache;
 
 		/** Used to determine whether to rebuild the entire property data. */
 		TArray<FSoftObjectPath> PreviousSelectedObjects;
 
 		TSharedRef<SWidget> CreatePropertiesView(const FArguments& InArgs);
 		
-		TSharedRef<FReplicatedPropertyData> AllocatePropertyData(FSoftClassPath OwningClass, FConcertPropertyChain PropertyChain);
-
-		/** Inits RootPropertyRowData from PropertyRowData. */
-		void BuildRootPropertyRowData();
 		/** Given the selected objects, determines whether they all have the same class and returns it if so. */
 		TOptional<FSoftClassPath> GetClassForPropertiesFromSelection(const TArray<FSoftObjectPath>& Objects) const;
-		void GetPropertyRowChildren(TSharedPtr<FReplicatedPropertyData> ReplicatedPropertyData, TFunctionRef<void(TSharedPtr<FReplicatedPropertyData>)> ProcessChild);
-		
-		// Utils
+		/** Sets how to display this widget */
 		void SetPropertyContent(EReplicatedPropertyContent Content) const;
 	};
 }
diff --git a/Engine/Plugins/Developer/Concert/ConcertUI/ConcertSharedSlate/Source/ConcertSharedSlate/Private/Replication/Editor/View/ObjectViewer/SReplicationStreamViewer.cpp b/Engine/Plugins/Developer/Concert/ConcertUI/ConcertSharedSlate/Source/ConcertSharedSlate/Private/Replication/Editor/View/ObjectViewer/SReplicationStreamViewer.cpp
--- a/Engine/Plugins/Developer/Concert/ConcertUI/ConcertSharedSlate/Source/ConcertSharedSlate/Private/Replication/Editor/View/ObjectViewer/SReplicationStreamViewer.cpp
+++ b/Engine/Plugins/Developer/Concert/ConcertUI/ConcertSharedSlate/Source/ConcertSharedSlate/Private/Replication/Editor/View/ObjectViewer/SReplicationStreamViewer.cpp
@@ -56,8 +56,14 @@ namespace UE::ConcertSharedSlate
 		return PropertySection->GetObjectsSelectedForPropertyEditing();
 	}
 
-	void SReplicationStreamViewer::SelectObjects(TConstArrayView<FSoftObjectPath> Objects)
+	void SReplicationStreamViewer::SelectObjects(TConstArrayView<FSoftObjectPath> Objects, bool bAtEndOfTick)
 	{
+		if (bHasRequestedObjectRefresh || bAtEndOfTick)
+		{
+			PendingToSelect = Objects;
+			return;
+		}
+		
 		TArray<TSharedPtr<FReplicatedObjectData>> NewSelectedItems; 
 		Algo::TransformIf(AllObjectRowData, NewSelectedItems, [&Objects](const TSharedPtr<FReplicatedObjectData>& ObjectData)
 			{
@@ -71,12 +77,19 @@ namespace UE::ConcertSharedSlate
 		}
 	}
 
-	void SReplicationStreamViewer::ExpandObjects(TConstArrayView<FSoftObjectPath> Objects, bool bRecursive)
+	void SReplicationStreamViewer::ExpandObjects(TConstArrayView<FSoftObjectPath> Objects, bool bRecursive, bool bAtEndOfTick)
 	{
 		if (Objects.IsEmpty())
 		{
 			return;
 		}
+
+		if (bHasRequestedObjectRefresh || bAtEndOfTick)
+		{
+			PendingToExpand = Objects;
+			bPendingExpandRecursively = bRecursive;
+			return;
+		}
 		
 		TArray<TSharedPtr<FReplicatedObjectData>> ItemsToExpand;
 		ItemsToExpand.Reserve(Objects.Num());
@@ -112,7 +125,10 @@ namespace UE::ConcertSharedSlate
 		// Items may have been removed this tick. However, selected items may not have been updated yet because STreeView processes item changes at the end of tick. 
 		SelectedItems.SetNum(Algo::RemoveIf(SelectedItems, [this](const TSharedPtr<FReplicatedObjectData>& ObjectData)
 		{
-			return !PropertiesModel->ContainsObjects({ ObjectData->GetObjectPath() });
+			const bool bIsInModel = PropertiesModel->ContainsObjects({ ObjectData->GetObjectPath() });
+			const TOptional<FSoftObjectPath> OwningActor = ObjectUtils::GetActorOf(ObjectData->GetObjectPath());
+			const bool bHierarchyIsInModel = OwningActor.IsSet() && PropertiesModel->ContainsObjects({ *OwningActor });
+			return !bIsInModel && !bHierarchyIsInModel;
 		}));
 		return SelectedItems;
 	}
@@ -131,6 +147,18 @@ namespace UE::ConcertSharedSlate
 			bHasRequestedPropertyRefresh = false;
 			RefreshPropertyData();
 		}
+
+		if (!PendingToSelect.IsEmpty())
+		{
+			SelectObjects(PendingToSelect);
+			PendingToSelect.Reset();
+		}
+
+		if (!PendingToExpand.IsEmpty())
+		{
+			ExpandObjects(PendingToExpand, bPendingExpandRecursively);
+			PendingToExpand.Reset();
+		}
 		
 		IReplicationStreamViewer::Tick(AllottedGeometry, InCurrentTime, InDeltaTime);
 	}
@@ -289,7 +317,7 @@ namespace UE::ConcertSharedSlate
 
 			// The tree view requires the item source to only contain the root items. Children are discovered via GetObjectRowChildren. We re-use GetObjectRowChildren to remove any non-root nodes.
 			BuildRootObjectRowData();
-			ReplicatedObjects->OnItemsChanged();
+			ReplicatedObjects->RequestRefilter();
 		}
 	}
 
diff --git a/Engine/Plugins/Developer/Concert/ConcertUI/ConcertSharedSlate/Source/ConcertSharedSlate/Private/Replication/Editor/View/ObjectViewer/SReplicationStreamViewer.h b/Engine/Plugins/Developer/Concert/ConcertUI/ConcertSharedSlate/Source/ConcertSharedSlate/Private/Replication/Editor/View/ObjectViewer/SReplicationStreamViewer.h
--- a/Engine/Plugins/Developer/Concert/ConcertUI/ConcertSharedSlate/Source/ConcertSharedSlate/Private/Replication/Editor/View/ObjectViewer/SReplicationStreamViewer.h
+++ b/Engine/Plugins/Developer/Concert/ConcertUI/ConcertSharedSlate/Source/ConcertSharedSlate/Private/Replication/Editor/View/ObjectViewer/SReplicationStreamViewer.h
@@ -101,9 +101,9 @@ namespace UE::ConcertSharedSlate
 		void RequestPropertyDataRefresh() { bHasRequestedPropertyRefresh = true; }
 
 		/** Selects the given objects. */
-		void SelectObjects(TConstArrayView<FSoftObjectPath> Objects);
+		void SelectObjects(TConstArrayView<FSoftObjectPath> Objects, bool bAtEndOfTick = false);
 		/** Expands the given objects, recursively if desired. */
-		void ExpandObjects(TConstArrayView<FSoftObjectPath> Objects, bool bRecursive);
+		void ExpandObjects(TConstArrayView<FSoftObjectPath> Objects, bool bRecursive, bool bAtEndOfTick = false);
 
 		/** @return Gets the root objects selected in the outliner; the subobject view chooses which of these objects (or their subobjects) end up in GetSelectedObjectShowingProperties. */
 		TArray<TSharedPtr<FReplicatedObjectData>> GetSelectedOutlinerObjects() const;
@@ -142,6 +142,9 @@ namespace UE::ConcertSharedSlate
 
 		bool bHasRequestedObjectRefresh = false;
 		bool bHasRequestedPropertyRefresh = false;
+		TArray<FSoftObjectPath> PendingToSelect;
+		TArray<FSoftObjectPath> PendingToExpand;
+		bool bPendingExpandRecursively;
 
 		static TSharedRef<FReplicatedObjectData> AllocateObjectData(FSoftObjectPath ObjectPath);
 
diff --git a/Engine/Plugins/Developer/Concert/ConcertUI/ConcertSharedSlate/Source/ConcertSharedSlate/Private/Replication/Editor/View/Tree/SReplicationTreeView.h b/Engine/Plugins/Developer/Concert/ConcertUI/ConcertSharedSlate/Source/ConcertSharedSlate/Private/Replication/Editor/View/Tree/SReplicationTreeView.h
--- a/Engine/Plugins/Developer/Concert/ConcertUI/ConcertSharedSlate/Source/ConcertSharedSlate/Private/Replication/Editor/View/Tree/SReplicationTreeView.h
+++ b/Engine/Plugins/Developer/Concert/ConcertUI/ConcertSharedSlate/Source/ConcertSharedSlate/Private/Replication/Editor/View/Tree/SReplicationTreeView.h
@@ -198,12 +198,6 @@ namespace UE::ConcertSharedSlate
 			}
 		}
 		
-		void OnItemsChanged()
-		{
-			// Re-filter everything. There should not be many items so filtering everything again should be fine
-			bFilterChanged = true;
-			bRequestedSort = true;
-		}
 		void RequestRefilter()
 		{
 			bFilterChanged = true;
diff --git a/Engine/Plugins/Developer/Concert/ConcertUI/ConcertSharedSlate/Source/ConcertSharedSlate/Public/Replication/Editor/View/ReplicationColumnsUtils.h b/Engine/Plugins/Developer/Concert/ConcertUI/ConcertSharedSlate/Source/ConcertSharedSlate/Public/Replication/Editor/View/ReplicationColumnsUtils.h
--- a/Engine/Plugins/Developer/Concert/ConcertUI/ConcertSharedSlate/Source/ConcertSharedSlate/Public/Replication/Editor/View/ReplicationColumnsUtils.h
+++ b/Engine/Plugins/Developer/Concert/ConcertUI/ConcertSharedSlate/Source/ConcertSharedSlate/Public/Replication/Editor/View/ReplicationColumnsUtils.h
@@ -7,8 +7,6 @@
 #include "Widgets/Input/SCheckBox.h"
 #include "Templates/SharedPointer.h"
 
-struct FConcertPropertyChain;
-
 namespace UE::ConcertSharedSlate
 {
 	class IEditableReplicationStreamModel;
