Remove comments for ticket that won't be done

diff --git a/Engine/Plugins/Developer/Concert/ConcertUI/ConcertSharedSlate/Source/ConcertSharedSlate/Private/Replication/Editor/Model/GenericReplicationStreamModel.cpp b/Engine/Plugins/Developer/Concert/ConcertUI/ConcertSharedSlate/Source/ConcertSharedSlate/Private/Replication/Editor/Model/GenericReplicationStreamModel.cpp
--- a/Engine/Plugins/Developer/Concert/ConcertUI/ConcertSharedSlate/Source/ConcertSharedSlate/Private/Replication/Editor/Model/GenericReplicationStreamModel.cpp
+++ b/Engine/Plugins/Developer/Concert/ConcertUI/ConcertSharedSlate/Source/ConcertSharedSlate/Private/Replication/Editor/Model/GenericReplicationStreamModel.cpp
@@ -210,8 +210,6 @@ namespace UE::ConcertSharedSlate
 			// Parent properties must also be added
 			// Not exactly efficient to iterate through the hierarchy for every removed item but it should be fine... Properties.Num() == 1 is the most common case
 			bAddedAtLeastOne |= Private::AddParentProperties(*Class, AddedProperty, ReplicatedProperties);
-			
-			// TODO UE-202079: Make sure to append FConcertPropertyChain::InternalContainerPropertyValueName if needed
 		}
 
 		if (bAddedAtLeastOne)
@@ -308,7 +306,6 @@ namespace UE::ConcertSharedSlate
 					TArray<FConcertPropertyChain>& Properties = ObjectInfo.PropertySelection.ReplicatedProperties;
 					Private::AddParentProperties(*ObjectClass, PropertyChain, Properties);
 					Properties.Emplace(MoveTemp(PropertyChain));
-					// TODO UE-202079: Make sure to append FConcertPropertyChain::InternalContainerPropertyValueName if needed
 				}
 			}
 			
