Refactor columns injection interface for the MU replication tree view.
- Columns now implement the IReplicationTreeColumn interface

diff --git a/Engine/Plugins/Developer/Concert/ConcertApp/MultiUserClient/Source/MultiUserClient/Private/Widgets/ActiveSession/Replication/Client/Multi/Columns/MultiStreamColumns.h b/Engine/Plugins/Developer/Concert/ConcertApp/MultiUserClient/Source/MultiUserClient/Private/Widgets/ActiveSession/Replication/Client/Multi/Columns/MultiStreamColumns.h
--- a/Engine/Plugins/Developer/Concert/ConcertApp/MultiUserClient/Source/MultiUserClient/Private/Widgets/ActiveSession/Replication/Client/Multi/Columns/MultiStreamColumns.h
+++ b/Engine/Plugins/Developer/Concert/ConcertApp/MultiUserClient/Source/MultiUserClient/Private/Widgets/ActiveSession/Replication/Client/Multi/Columns/MultiStreamColumns.h
@@ -2,7 +2,8 @@
 
 #pragma once
 
-#include "Replication/Editor/View/PredefinedReplicationColumns.h"
+#include "Replication/Editor/View/Column/IPropertyTreeColumn.h"
+#include "Replication/Editor/View/Column/IObjectTreeColumn.h"
 
 #include "Delegates/Delegate.h"
 #include "Misc/EnumClassFlags.h"
@@ -47,7 +48,7 @@ namespace UE::MultiUserClient::MultiStreamColumns
 	 * 
 	 * @return A checkbox for controlling the authority of the object in the row
 	 */
-	ConcertSharedSlate::ReplicationColumns::FReplicationTopLevelObjectColumn ReplicationToggle(
+	ConcertSharedSlate::FObjectColumnEntry ReplicationToggle(
 		TSharedRef<IConcertClient> ConcertClient,
 		TAttribute<ConcertSharedSlate::IObjectHierarchyModel*> ObjectHierarchyModelAttribute,
 		FReplicationClientManager& ClientManager,
@@ -64,7 +65,7 @@ namespace UE::MultiUserClient::MultiStreamColumns
 	 * @param ClientManager Used to access all clients for display in the combo box drop-down
 	 * @param ColumnsSortPriority The order relative to the other columns
 	 */
-	ConcertSharedSlate::ReplicationColumns::FReplicationTopLevelObjectColumn ReassignOwnership(
+	ConcertSharedSlate::FObjectColumnEntry ReassignOwnership(
 		TSharedRef<IConcertClient> ConcertClient,
 		TAttribute<TSharedPtr<ConcertSharedSlate::IMultiReplicationStreamEditor>> MultiStreamModelAttribute,
 		TAttribute<ConcertSharedSlate::IObjectHierarchyModel*> ObjectHierarchyModelAttribute,
@@ -83,7 +84,7 @@ namespace UE::MultiUserClient::MultiStreamColumns
 	 * 
 	 * @return A column that spawns a combo box for assigning properties
 	 */
-	ConcertSharedSlate::ReplicationColumns::FReplicationPropertyColumn AssignPropertyColumn(
+	ConcertSharedSlate::FPropertyColumnEntry AssignPropertyColumn(
 		TAttribute<TSharedPtr<ConcertSharedSlate::IMultiReplicationStreamEditor>> MultiStreamEditor,
 		TSharedRef<IConcertClient> ConcertClient,
 		FReplicationClientManager& ClientManager,
diff --git a/Engine/Plugins/Developer/Concert/ConcertApp/MultiUserClient/Source/MultiUserClient/Private/Widgets/ActiveSession/Replication/Client/Multi/Columns/MultiStreamColumns_AssignProperties.cpp b/Engine/Plugins/Developer/Concert/ConcertApp/MultiUserClient/Source/MultiUserClient/Private/Widgets/ActiveSession/Replication/Client/Multi/Columns/MultiStreamColumns_AssignProperties.cpp
--- a/Engine/Plugins/Developer/Concert/ConcertApp/MultiUserClient/Source/MultiUserClient/Private/Widgets/ActiveSession/Replication/Client/Multi/Columns/MultiStreamColumns_AssignProperties.cpp
+++ b/Engine/Plugins/Developer/Concert/ConcertApp/MultiUserClient/Source/MultiUserClient/Private/Widgets/ActiveSession/Replication/Client/Multi/Columns/MultiStreamColumns_AssignProperties.cpp
@@ -11,6 +11,8 @@
 #include "Replication/Editor/View/IMultiReplicationStreamEditor.h"
 #include "Replication/Editor/View/IReplicationStreamEditor.h"
 #include "SAssignPropertyComboBox.h"
+#include "Replication/Editor/View/Column/IPropertyTreeColumn.h"
+#include "Replication/Editor/View/Column/ReplicationColumnDelegates.h"
 #include "Widgets/ActiveSession/Replication/Client/ClientUtils.h"
 
 #define LOCTEXT_NAMESPACE "AssignPropertyColumn"
@@ -72,7 +74,7 @@ namespace UE::MultiUserClient::MultiStreamColumns
 		}
 	}
 	
-	ConcertSharedSlate::ReplicationColumns::FReplicationPropertyColumn AssignPropertyColumn(
+	ConcertSharedSlate::FPropertyColumnEntry AssignPropertyColumn(
 		TAttribute<TSharedPtr<ConcertSharedSlate::IMultiReplicationStreamEditor>> MultiStreamEditor,
 		TSharedRef<IConcertClient> ConcertClient,
 		FReplicationClientManager& ClientManager,
@@ -80,57 +82,85 @@ namespace UE::MultiUserClient::MultiStreamColumns
 		)
 	{
 		using namespace ConcertSharedSlate;
-		using namespace ConcertSharedSlate::ReplicationColumns;
-		check(MultiStreamEditor.IsBound() || MultiStreamEditor.IsSet());
-
-		const auto PopulateSearch = [MultiStreamEditor, &ClientManager, ConcertClient](const FReplicatedPropertyData& Data, TArray<FString>& InOutSearchStrings)
+		class FPropertyColumn_AssignPropertyColumn : public IPropertyTreeColumn
 		{
-			AssignPropertyColumnUtils::ForEachStreamAssignedTo(*MultiStreamEditor.Get(), Data.GetProperty(),
-				[&ClientManager, &ConcertClient, &InOutSearchStrings](const TSharedRef<IEditableReplicationStreamModel>& Stream)
-				{
-					InOutSearchStrings.Add(AssignPropertyColumnUtils::GetClientDisplayText(*ConcertClient, ClientManager, *Stream));
-				});
-		};
-		const auto GenerateWidgetColumn = [MultiStreamEditor, ConcertClient, &ClientManager](const FReplicationPropertyColumn::FBuildArgs& Args)
-		{
-			const TArray<FSoftObjectPath> DisplayedObjects = MultiStreamEditor.Get()->GetEditorBase().GetObjectsBeingPropertyEdited();
-			return SNew(SAssignPropertyComboBox, MultiStreamEditor.Get().ToSharedRef(), ConcertClient, ClientManager)
-				.DisplayedProperty(Args.RowData.GetProperty())
-				.EditedObjects(DisplayedObjects)
-				.HighlightText(Args.HighlightText)
-				.OnOptionSelected_Lambda([MultiStreamEditor](auto)
-				{
-					if (const TSharedPtr<IMultiReplicationStreamEditor> Editor = MultiStreamEditor.Get())
+		public:
+
+			FPropertyColumn_AssignPropertyColumn(
+				TAttribute<TSharedPtr<IMultiReplicationStreamEditor>> MultiStreamEditor,
+				TSharedRef<IConcertClient> ConcertClient,
+				FReplicationClientManager& ClientManager
+				)
+				: MultiStreamEditor(MoveTemp(MultiStreamEditor))
+				, ConcertClient(MoveTemp(ConcertClient))
+				, ClientManager(ClientManager)
+			{}
+			
+			virtual SHeaderRow::FColumn::FArguments CreateHeaderRowArgs() const override
+			{
+				return SHeaderRow::Column(AssignPropertyColumnId)
+					.DefaultLabel(LOCTEXT("Owner.Label", "Assigned Client"))
+					.ToolTipText(LOCTEXT("Owner.ToolTip", "Client that should replicate this property"))
+					.FillSized(FMultiUserReplicationStyle::Get()->GetFloat(TEXT("AllClients.Property.OwnerColumnWidth")));
+			}
+			
+			virtual TSharedRef<SWidget> GenerateColumnWidget(const FBuildArgs& InArgs) override
+			{
+				const TArray<FSoftObjectPath> DisplayedObjects = MultiStreamEditor.Get()->GetEditorBase().GetObjectsBeingPropertyEdited();
+				return SNew(SAssignPropertyComboBox, MultiStreamEditor.Get().ToSharedRef(), ConcertClient, ClientManager)
+					.DisplayedProperty(InArgs.RowItem.RowData.GetProperty())
+					.EditedObjects(DisplayedObjects)
+					.HighlightText(InArgs.HighlightText)
+					.OnOptionSelected_Lambda([this](auto)
 					{
-						Editor->GetEditorBase().RequestPropertyColumnResort(AssignPropertyColumnId);
-					}
-				});
-		};
-		const auto IsLessThan = [MultiStreamEditor, ConcertClient, &ClientManager](const FReplicatedPropertyData& Left, const FReplicatedPropertyData& Right)
-		{
-			const TArray<FSoftObjectPath> DisplayedObjects = MultiStreamEditor.Get()->GetEditorBase().GetObjectsBeingPropertyEdited();
-			const TOptional<FString> LeftClientDisplayString = SAssignPropertyComboBox::GetDisplayString(ConcertClient, ClientManager, Left.GetProperty(), DisplayedObjects);
-			const TOptional<FString> RightClientDisplayString = SAssignPropertyComboBox::GetDisplayString(ConcertClient, ClientManager, Right.GetProperty(), DisplayedObjects);
+						if (const TSharedPtr<IMultiReplicationStreamEditor> Editor = MultiStreamEditor.Get())
+						{
+							Editor->GetEditorBase().RequestPropertyColumnResort(AssignPropertyColumnId);
+						}
+					});
+			}
 			
-			if (LeftClientDisplayString && RightClientDisplayString)
+			virtual void PopulateSearchString(const FPropertyTreeRowContext& InItem, TArray<FString>& InOutSearchStrings) const override
 			{
-				return *LeftClientDisplayString < *RightClientDisplayString;
+				AssignPropertyColumnUtils::ForEachStreamAssignedTo(*MultiStreamEditor.Get(), InItem.RowData.GetProperty(),
+					[this, &InOutSearchStrings](const TSharedRef<IEditableReplicationStreamModel>& Stream)
+					{
+						InOutSearchStrings.Add(AssignPropertyColumnUtils::GetClientDisplayText(*ConcertClient, ClientManager, *Stream));
+					});
+			}
+
+			virtual bool CanBeSorted() const override { return true; }
+			virtual bool IsLessThan(const FPropertyTreeRowContext& Left, const FPropertyTreeRowContext& Right) const override
+			{
+				const TArray<FSoftObjectPath> DisplayedObjects = MultiStreamEditor.Get()->GetEditorBase().GetObjectsBeingPropertyEdited();
+				const TOptional<FString> LeftClientDisplayString = SAssignPropertyComboBox::GetDisplayString(ConcertClient, ClientManager, Left.RowData.GetProperty(), DisplayedObjects);
+				const TOptional<FString> RightClientDisplayString = SAssignPropertyComboBox::GetDisplayString(ConcertClient, ClientManager, Right.RowData.GetProperty(), DisplayedObjects);
+			
+				if (LeftClientDisplayString && RightClientDisplayString)
+				{
+					return *LeftClientDisplayString < *RightClientDisplayString;
+				}
+				// Our rule: set < unset. This way unassigned appears last.
+				return LeftClientDisplayString.IsSet() && !RightClientDisplayString.IsSet();
 			}
-			// Our rule: set < unset. This way unassigned appears last.
-			return LeftClientDisplayString.IsSet() && !RightClientDisplayString.IsSet();
+
+		private:
+
+			const TAttribute<TSharedPtr<IMultiReplicationStreamEditor>> MultiStreamEditor;
+			const TSharedRef<IConcertClient> ConcertClient;
+			FReplicationClientManager& ClientManager;
 		};
 		
-		return FReplicationPropertyColumn(
-			typename FReplicationPropertyColumn::FArguments()
-				.PopulateSearchItems_Lambda(PopulateSearch)
-				.GenerateWidgetColumn_Lambda(GenerateWidgetColumn)
-				.IsLessThan_Lambda(IsLessThan)
-				.ColumnSortOrder(ColumnsSortPriority),
-			SHeaderRow::Column(AssignPropertyColumnId)
-				.DefaultLabel(LOCTEXT("Owner.Label", "Assigned Client"))
-				.ToolTipText(LOCTEXT("Owner.ToolTip", "Client that should replicate this property"))
-				.FillSized(FMultiUserReplicationStyle::Get()->GetFloat(TEXT("AllClients.Property.OwnerColumnWidth")))
-			);
+		check(MultiStreamEditor.IsBound() || MultiStreamEditor.IsSet());
+		return {
+			TReplicationColumnDelegates<FPropertyTreeRowContext>::FCreateColumn::CreateLambda(
+				[MultiStreamEditor = MoveTemp(MultiStreamEditor), ConcertClient = MoveTemp(ConcertClient), &ClientManager]()
+				{
+					return MakeShared<FPropertyColumn_AssignPropertyColumn>(MultiStreamEditor, ConcertClient, ClientManager);
+				}),
+			AssignPropertyColumnId,
+			{ ColumnsSortPriority }
+		};
 	}
 }
 
diff --git a/Engine/Plugins/Developer/Concert/ConcertApp/MultiUserClient/Source/MultiUserClient/Private/Widgets/ActiveSession/Replication/Client/Multi/Columns/MultiStreamColumns_ReassignOwnership.cpp b/Engine/Plugins/Developer/Concert/ConcertApp/MultiUserClient/Source/MultiUserClient/Private/Widgets/ActiveSession/Replication/Client/Multi/Columns/MultiStreamColumns_ReassignOwnership.cpp
--- a/Engine/Plugins/Developer/Concert/ConcertApp/MultiUserClient/Source/MultiUserClient/Private/Widgets/ActiveSession/Replication/Client/Multi/Columns/MultiStreamColumns_ReassignOwnership.cpp
+++ b/Engine/Plugins/Developer/Concert/ConcertApp/MultiUserClient/Source/MultiUserClient/Private/Widgets/ActiveSession/Replication/Client/Multi/Columns/MultiStreamColumns_ReassignOwnership.cpp
@@ -14,7 +14,7 @@ namespace UE::MultiUserClient::MultiStreamColumns
 {
 	const FName ReassignOwnershipColumnId(TEXT("ReassignOwnershipColumn"));
 	
-	ConcertSharedSlate::ReplicationColumns::FReplicationTopLevelObjectColumn ReassignOwnership(
+	ConcertSharedSlate::FObjectColumnEntry ReassignOwnership(
 		TSharedRef<IConcertClient> ConcertClient,
 		TAttribute<TSharedPtr<ConcertSharedSlate::IMultiReplicationStreamEditor>> MultiStreamModelAttribute,
 		TAttribute<ConcertSharedSlate::IObjectHierarchyModel*> ObjectHierarchyModelAttribute,
@@ -23,51 +23,88 @@ namespace UE::MultiUserClient::MultiStreamColumns
 		const int32 ColumnsSortPriority
 		)
 	{
-		using namespace ConcertSharedSlate::ReplicationColumns;
+		class FObjectColumn_ReassignOwnership : public ConcertSharedSlate::IObjectTreeColumn
+		{
+		public:
 
-		auto MakeWidget =
-			[ConcertClient, MultiStreamModelAttribute = MoveTemp(MultiStreamModelAttribute), ObjectHierarchyModelAttribute = MoveTemp(ObjectHierarchyModelAttribute), &ReassignmentLogic, &ClientManager]
-			(const FReplicationTopLevelObjectColumn::FBuildArgs& InArgs)
+			FObjectColumn_ReassignOwnership(
+				TSharedRef<IConcertClient> ConcertClient,
+				TAttribute<TSharedPtr<ConcertSharedSlate::IMultiReplicationStreamEditor>> MultiStreamModelAttribute,
+				TAttribute<ConcertSharedSlate::IObjectHierarchyModel*> ObjectHierarchyModelAttribute,
+				FReassignObjectPropertiesLogic& ReassignmentLogic,
+				const FReplicationClientManager& ClientManager
+				)
+				: ConcertClient(MoveTemp(ConcertClient))
+				, MultiStreamModelAttribute(MoveTemp(MultiStreamModelAttribute))
+				, ObjectHierarchyModelAttribute(MoveTemp(ObjectHierarchyModelAttribute))
+				, ReassignmentLogic(ReassignmentLogic)
+				, ClientManager(ClientManager)
+			{}
+			
+			virtual SHeaderRow::FColumn::FArguments CreateHeaderRowArgs() const override
+			{
+				return SHeaderRow::Column(ReassignOwnershipColumnId)
+					.DefaultLabel(LOCTEXT("Owner.Label", "Assigned Clients"))
+					.ToolTipText(LOCTEXT("Owner.ToolTip", "Clients that have registered properties for an object"))
+					.FillSized(FMultiUserReplicationStyle::Get()->GetFloat(TEXT("AllClients.Object.OwnerColumnWidth")));
+			}
+			
+			virtual TSharedRef<SWidget> GenerateColumnWidget(const FBuildArgs& InArgs) override
 			{
 				return SNew(SReassignObjectComboBox, ConcertClient, ReassignmentLogic, ClientManager)
-					.ManagedObject(InArgs.RowData.GetObjectPath())
+					.ManagedObject(InArgs.RowItem.RowData.GetObjectPath())
 					.ObjectHierarchyModel(ObjectHierarchyModelAttribute)
 					.HighlightText(InArgs.HighlightText)
-					.OnReassignAllOptionClicked_Lambda([MultiStreamModelAttribute](auto)
+					.OnReassignAllOptionClicked_Lambda([this](auto)
 					{
 						if (const TSharedPtr<ConcertSharedSlate::IMultiReplicationStreamEditor> Model = MultiStreamModelAttribute.Get())
 						{
 							Model->GetEditorBase().RequestObjectColumnResort(ReassignOwnershipColumnId);
 						}
 					});
-			};
-		auto IsLessThan = [ConcertClient, &ReassignmentLogic](const ConcertSharedSlate::FReplicatedObjectData& Left, const ConcertSharedSlate::FReplicatedObjectData& Right)
-		{
-			const TOptional<FString> LeftClientDisplayString = SReassignObjectComboBox::GetDisplayString(ConcertClient, ReassignmentLogic, Left.GetObjectPath());
-			const TOptional<FString> RightClientDisplayString = SReassignObjectComboBox::GetDisplayString(ConcertClient, ReassignmentLogic, Right.GetObjectPath());
-			
-			if (LeftClientDisplayString && RightClientDisplayString)
+			}
+			virtual void PopulateSearchString(const ConcertSharedSlate::FObjectTreeRowContext& InItem, TArray<FString>& InOutSearchStrings) const override
 			{
-				return *LeftClientDisplayString < *RightClientDisplayString;
+				SReassignObjectComboBox::PopulateSearchTerms(
+					*ConcertClient->GetCurrentSession(),
+					ReassignmentLogic,
+					InItem.RowData.GetObjectPath(),
+					InOutSearchStrings
+					);
 			}
-			// Our rule: set < unset. This way unassigned appears last.
-			return LeftClientDisplayString.IsSet() && !RightClientDisplayString.IsSet();
+
+			virtual bool CanBeSorted() const override { return true; }
+			virtual bool IsLessThan(const ConcertSharedSlate::FObjectTreeRowContext& Left, const ConcertSharedSlate::FObjectTreeRowContext& Right) const override
+			{
+				const TOptional<FString> LeftClientDisplayString = SReassignObjectComboBox::GetDisplayString(ConcertClient, ReassignmentLogic, Left.RowData.GetObjectPath());
+				const TOptional<FString> RightClientDisplayString = SReassignObjectComboBox::GetDisplayString(ConcertClient, ReassignmentLogic, Right.RowData.GetObjectPath());
+			
+				if (LeftClientDisplayString && RightClientDisplayString)
+				{
+					return *LeftClientDisplayString < *RightClientDisplayString;
+				}
+				// Our rule: set < unset. This way unassigned appears last.
+				return LeftClientDisplayString.IsSet() && !RightClientDisplayString.IsSet();
+			}
+
+		private:
+			
+			const TSharedRef<IConcertClient> ConcertClient;
+			const TAttribute<TSharedPtr<ConcertSharedSlate::IMultiReplicationStreamEditor>> MultiStreamModelAttribute;
+			const TAttribute<ConcertSharedSlate::IObjectHierarchyModel*> ObjectHierarchyModelAttribute;
+			FReassignObjectPropertiesLogic& ReassignmentLogic;
+			const FReplicationClientManager& ClientManager;
 		};
-		
-		return FReplicationTopLevelObjectColumn(
-			FReplicationTopLevelObjectColumn::FArguments()
-				.GenerateWidgetColumn_Lambda(MoveTemp(MakeWidget))
-				.PopulateSearchItems_Lambda([ConcertClient, &ReassignmentLogic](const ConcertSharedSlate::FReplicatedObjectData& ObjectData, TArray<FString>& InOutSearchStrings)
+
+		return {
+			ConcertSharedSlate::TReplicationColumnDelegates<ConcertSharedSlate::FObjectTreeRowContext>::FCreateColumn::CreateLambda(
+				[ConcertClient = MoveTemp(ConcertClient), MultiStreamModelAttribute = MoveTemp(MultiStreamModelAttribute), ObjectHierarchyModelAttribute = MoveTemp(ObjectHierarchyModelAttribute), &ReassignmentLogic, &ClientManager]()
 				{
-					SReassignObjectComboBox::PopulateSearchTerms(*ConcertClient->GetCurrentSession(), ReassignmentLogic, ObjectData.GetObjectPath(), InOutSearchStrings);
-				})
-				.IsLessThan_Lambda(MoveTemp(IsLessThan))
-				.ColumnSortOrder(ColumnsSortPriority),
-			SHeaderRow::Column(ReassignOwnershipColumnId)
-				.DefaultLabel(LOCTEXT("Owner.Label", "Assigned Clients"))
-				.ToolTipText(LOCTEXT("Owner.ToolTip", "Clients that have registered properties for an object"))
-				.FillSized(FMultiUserReplicationStyle::Get()->GetFloat(TEXT("AllClients.Object.OwnerColumnWidth")))
-			);
+					return MakeShared<FObjectColumn_ReassignOwnership>(ConcertClient, MultiStreamModelAttribute, ObjectHierarchyModelAttribute, ReassignmentLogic, ClientManager);
+				}),
+			ReassignOwnershipColumnId,
+			{ ColumnsSortPriority }
+		};
 	}
 }
 
diff --git a/Engine/Plugins/Developer/Concert/ConcertApp/MultiUserClient/Source/MultiUserClient/Private/Widgets/ActiveSession/Replication/Client/Multi/Columns/MultiStreamColumns_ReplicationToggle.cpp b/Engine/Plugins/Developer/Concert/ConcertApp/MultiUserClient/Source/MultiUserClient/Private/Widgets/ActiveSession/Replication/Client/Multi/Columns/MultiStreamColumns_ReplicationToggle.cpp
--- a/Engine/Plugins/Developer/Concert/ConcertApp/MultiUserClient/Source/MultiUserClient/Private/Widgets/ActiveSession/Replication/Client/Multi/Columns/MultiStreamColumns_ReplicationToggle.cpp
+++ b/Engine/Plugins/Developer/Concert/ConcertApp/MultiUserClient/Source/MultiUserClient/Private/Widgets/ActiveSession/Replication/Client/Multi/Columns/MultiStreamColumns_ReplicationToggle.cpp
@@ -4,6 +4,7 @@
 
 #include "MultiUserReplicationStyle.h"
 #include "SReplicationMultiToggleCheckbox.h"
+#include "Replication/Editor/View/Column/IObjectTreeColumn.h"
 
 #include "Widgets/Layout/SBox.h"
 #include "Widgets/Views/SHeaderRow.h"
@@ -14,33 +15,62 @@ namespace UE::MultiUserClient::MultiStreamColumns
 {
 	const FName ReplicationToggleColumnId(TEXT("ReplicationToggleColumn"));
 	
-	ConcertSharedSlate::ReplicationColumns::FReplicationTopLevelObjectColumn ReplicationToggle(
+	ConcertSharedSlate::FObjectColumnEntry ReplicationToggle(
 		TSharedRef<IConcertClient> ConcertClient,
 		TAttribute<ConcertSharedSlate::IObjectHierarchyModel*> ObjectHierarchyModelAttribute,
 		FReplicationClientManager& ClientManager,
 		const int32 ColumnsSortPriority
 		)
 	{
-		using namespace ConcertSharedSlate::ReplicationColumns;
-		
-		return FReplicationTopLevelObjectColumn(
-			FReplicationTopLevelObjectColumn::FArguments()
-				.GenerateWidgetColumn_Lambda([ConcertClient = MoveTemp(ConcertClient), ObjectHierarchyModelAttribute = MoveTemp(ObjectHierarchyModelAttribute), &ClientManager](const FReplicationTopLevelObjectColumn::FBuildArgs& InArgs)
-				{
-					return SNew(SBox)
+		class FObjectColumn_ReplicationToggle : public ConcertSharedSlate::IObjectTreeColumn
+		{
+		public:
+
+			FObjectColumn_ReplicationToggle(
+				TSharedRef<IConcertClient> ConcertClient,
+				TAttribute<ConcertSharedSlate::IObjectHierarchyModel*> ObjectHierarchyModelAttribute,
+				FReplicationClientManager& ClientManager
+				)
+				: ConcertClient(MoveTemp(ConcertClient))
+				, ObjectHierarchyModelAttribute(MoveTemp(ObjectHierarchyModelAttribute))
+				, ClientManager(ClientManager)
+			{}
+			
+			virtual SHeaderRow::FColumn::FArguments CreateHeaderRowArgs() const override
+			{
+				return SHeaderRow::Column(ReplicationToggleColumnId)
+					.DefaultLabel(FText::GetEmpty())
+					.ToolTipText(LOCTEXT("Replicates.ToolTip", "Assign properties first.\nControls whether the object should replicate."))
+					.FixedWidth(FMultiUserReplicationStyle::Get()->GetFloat(TEXT("AllClients.Object.ReplicationToggle")));
+			}
+			
+			virtual TSharedRef<SWidget> GenerateColumnWidget(const FBuildArgs& InArgs) override
+			{
+				return SNew(SBox)
 					.HAlign(HAlign_Left) // Warning icon is sometimes collapsed - we don't want the widget to be centered
 					[
 						SNew(SReplicationMultiToggleCheckbox, ClientManager, ConcertClient)
-						.Object(InArgs.RowData.GetObjectPath())
+						.Object(InArgs.RowItem.RowData.GetObjectPath())
 						.ObjectHierarchyModel(ObjectHierarchyModelAttribute)
 					];
-				})
-				.ColumnSortOrder(ColumnsSortPriority),
-			SHeaderRow::Column(ReplicationToggleColumnId)
-				.DefaultLabel(FText::GetEmpty())
-				.ToolTipText(LOCTEXT("Replicates.ToolTip", "Assign properties first.\nControls whether the object should replicate."))
-				.FixedWidth(FMultiUserReplicationStyle::Get()->GetFloat(TEXT("AllClients.Object.ReplicationToggle")))
-			);
+			}
+
+		private:
+
+			const TSharedRef<IConcertClient> ConcertClient;
+			const TAttribute<ConcertSharedSlate::IObjectHierarchyModel*> ObjectHierarchyModelAttribute;
+			FReplicationClientManager& ClientManager;
+		};
+
+		return {
+			ConcertSharedSlate::TReplicationColumnDelegates<ConcertSharedSlate::FObjectTreeRowContext>::FCreateColumn::CreateLambda(
+				[ConcertClient = MoveTemp(ConcertClient), ObjectHierarchyModelAttribute = MoveTemp(ObjectHierarchyModelAttribute), &ClientManager]()
+				{
+					return MakeShared<FObjectColumn_ReplicationToggle>(ConcertClient, ObjectHierarchyModelAttribute, ClientManager);
+				}),
+			ReplicationToggleColumnId,
+			{ ColumnsSortPriority }
+		};
 	}
 }
 
diff --git a/Engine/Plugins/Developer/Concert/ConcertApp/MultiUserClient/Source/MultiUserClient/Private/Widgets/ActiveSession/Replication/Client/Multi/SMultiClientView.cpp b/Engine/Plugins/Developer/Concert/ConcertApp/MultiUserClient/Source/MultiUserClient/Private/Widgets/ActiveSession/Replication/Client/Multi/SMultiClientView.cpp
--- a/Engine/Plugins/Developer/Concert/ConcertApp/MultiUserClient/Source/MultiUserClient/Private/Widgets/ActiveSession/Replication/Client/Multi/SMultiClientView.cpp
+++ b/Engine/Plugins/Developer/Concert/ConcertApp/MultiUserClient/Source/MultiUserClient/Private/Widgets/ActiveSession/Replication/Client/Multi/SMultiClientView.cpp
@@ -85,7 +85,7 @@ namespace UE::MultiUserClient
 		
 		ConcertClientSharedSlate::FFilterablePropertyTreeViewParams TreeViewParams
 		{
-			.PropertyColumns =
+			.AdditionalPropertyColumns =
 			{
 				ReplicationColumns::Property::LabelColumn(),
 				ReplicationColumns::Property::TypeColumn(),
@@ -111,7 +111,7 @@ namespace UE::MultiUserClient
 			.ObjectHierarchy = ObjectHierarchy, // This makes actors have children in the top view
 			.NameModel = ConcertClientSharedSlate::CreateEditorObjectNameModel(), // This makes actors use their labels, and components use the names given in the BP editor
 			.OnExtendObjectsContextMenu = FExtendObjectMenu::CreateSP(this, &SMultiClientView::ExtendObjectContextMenu),
-			.AdditionalObjectColumns =
+			.ObjectColumns =
 			{
 				MultiStreamColumns::ReplicationToggle(InConcertClient, ObjecHierarchyAttribute, InClientManager),
 				MultiStreamColumns::ReassignOwnership(InConcertClient, MultiStreamEditorAttribute, ObjecHierarchyAttribute, InClientManager.GetReassignmentLogic(), InClientManager)
diff --git a/Engine/Plugins/Developer/Concert/ConcertApp/MultiUserClient/Source/MultiUserClient/Private/Widgets/ActiveSession/Replication/Client/Single/Columns/SingleClientColumns.h b/Engine/Plugins/Developer/Concert/ConcertApp/MultiUserClient/Source/MultiUserClient/Private/Widgets/ActiveSession/Replication/Client/Single/Columns/SingleClientColumns.h
--- a/Engine/Plugins/Developer/Concert/ConcertApp/MultiUserClient/Source/MultiUserClient/Private/Widgets/ActiveSession/Replication/Client/Single/Columns/SingleClientColumns.h
+++ b/Engine/Plugins/Developer/Concert/ConcertApp/MultiUserClient/Source/MultiUserClient/Private/Widgets/ActiveSession/Replication/Client/Single/Columns/SingleClientColumns.h
@@ -2,7 +2,8 @@
 
 #pragma once
 
-#include "Replication/Editor/View/PredefinedReplicationColumns.h"
+#include "Replication/Editor/View/Column/IObjectTreeColumn.h"
+#include "Replication/Editor/View/Column/IPropertyTreeColumn.h"
 #include "Widgets/DeclarativeSyntaxSupport.h"
 #include "Widgets/SCompoundWidget.h"
 
@@ -40,33 +41,17 @@ namespace UE::MultiUserClient::SingleClientColumns
 
 
 	/********** Toggle authority **********/
-	extern const FName ToggleTopLevelAuthorityColumnId;
-
-	/**
-	 * Checkbox placed in the top-level view.
-	 * It gives / removes authority for the top-level object and all of its subobjects.
-	 *
-	 * @param ClientStreamModel Used to discover the client's registered subobjects
-	 * @param ChangeTracker Used to determine the checkbox state and change authority.
-	 * @param SubmissionWorkflow Used to determine whether changing authority is at all enabled
-	 * @return Column that can be placed in the table
-	 */
-	ConcertSharedSlate::ReplicationColumns::FReplicationTopLevelObjectColumn ToggleTopLevelAuthority(
-		ConcertSharedSlate::IReplicationStreamModel& ClientStreamModel,
-		FAuthorityChangeTracker& ChangeTracker,
-		ISubmissionWorkflow& SubmissionWorkflow
-		);
+	extern const FName ToggleObjectAuthorityColumnId;
 
 	/**
 	 * Checkbox placed in the subobject view
 	 * It gives / removes authority for the object it is placed next to.
 	 * 
-	 * @param AuthoritySynchronizer Determines whether the client has authority or not.
 	 * @param ChangeTracker Used to determine the checkbox state and change authority.
-	 * @param AuthoritySynchronizer Used to determine whether changing authority is at all enabled
+	 * @param SubmissionWorkflow
 	 * @return Column that can be placed in the table
 	 */
-	ConcertSharedSlate::ReplicationColumns::FReplicationTopLevelObjectColumn ToggleObjectAuthority(
+	ConcertSharedSlate::FObjectColumnEntry ToggleObjectAuthority(
 		FAuthorityChangeTracker& ChangeTracker,
 		ISubmissionWorkflow& SubmissionWorkflow
 		);
@@ -75,19 +60,6 @@ namespace UE::MultiUserClient::SingleClientColumns
 	/********** Owner **********/
 	extern const FName OwnerOfSubobjectColumnId;
 	extern const FName OwnerOfPropertyColumnId;
-	
-	/**
-	 * Displays the owner of a top-level object.
-	 * @param InClient The local Concert client used to look up other client display info
-	 * @param InAuthorityCache Used to determine which client owns the property
-	 * @param InObjectModel Used to get subobjects for a top-level object
-	 * @return Column that can be placed in the table
-	 */
-	ConcertSharedSlate::ReplicationColumns::FReplicationTopLevelObjectColumn OwnerOfTopLevelObject(
-		const TSharedRef<IConcertClient>& InClient,
-		FGlobalAuthorityCache& InAuthorityCache,
-		ConcertSharedSlate::IReplicationStreamModel& InObjectModel
-		);
 
 	/**
 	 * Displays the owner of a subobject.
@@ -95,8 +67,8 @@ namespace UE::MultiUserClient::SingleClientColumns
 	 * @param InAuthorityCache Used to determine which client owns the property
 	 * @return Column that can be placed in the table
 	 */
-	ConcertSharedSlate::ReplicationColumns::FReplicationTopLevelObjectColumn OwnerOfObject(
-		const TSharedRef<IConcertClient>& InClient,
+	ConcertSharedSlate::FObjectColumnEntry OwnerOfObject(
+		TSharedRef<IConcertClient> InClient,
 		FGlobalAuthorityCache& InAuthorityCache
 		);
 	
@@ -107,16 +79,15 @@ namespace UE::MultiUserClient::SingleClientColumns
 	 * @param InViewerAttribute Used to determine which objects the property box is displaying
 	 * @return Column that can be placed in the table
 	 */
-	ConcertSharedSlate::ReplicationColumns::FReplicationPropertyColumn OwnerOfProperty(
-		const TSharedRef<IConcertClient>& InClient,
+	ConcertSharedSlate::FPropertyColumnEntry OwnerOfProperty(
+		TSharedRef<IConcertClient> InClient,
 		FGlobalAuthorityCache& InAuthorityCache,
-		const TAttribute<const ConcertSharedSlate::IReplicationStreamViewer*>& InViewerAttribute
+		TAttribute<const ConcertSharedSlate::IReplicationStreamViewer*> InViewerAttribute
 		);
 
 	
 	/********** Conflict warning **********/
 	extern const FName ConflictWarningTopLevelObjectColumnId;
-	extern const FName ConflictWarningSubobjectColumnId;
 	extern const FName ConflictWarningPropertyColumnId;
 	
 	/**
@@ -128,7 +99,7 @@ namespace UE::MultiUserClient::SingleClientColumns
 	 * 
 	 * @return Column that can be placed in the table
 	 */
-	ConcertSharedSlate::ReplicationColumns::FReplicationTopLevelObjectColumn ConflictWarningForObject(
+	ConcertSharedSlate::FObjectColumnEntry ConflictWarningForObject(
 		TSharedRef<IConcertClient> InClient,
 		FGlobalAuthorityCache& InAuthorityCache,
 		const FGuid& ClientId
@@ -144,7 +115,7 @@ namespace UE::MultiUserClient::SingleClientColumns
 	 * 
 	 * @return Column that can be placed in the table
 	 */
-	ConcertSharedSlate::ReplicationColumns::FReplicationPropertyColumn ConflictWarningForProperty(
+	ConcertSharedSlate::FPropertyColumnEntry ConflictWarningForProperty(
 		TSharedRef<IConcertClient> InClient,
 		TAttribute<const ConcertSharedSlate::IReplicationStreamViewer*> InViewer,
 		FGlobalAuthorityCache& InAuthorityCache,
diff --git a/Engine/Plugins/Developer/Concert/ConcertApp/MultiUserClient/Source/MultiUserClient/Private/Widgets/ActiveSession/Replication/Client/Single/Columns/SingleClientColumns_ConflictWarning.cpp b/Engine/Plugins/Developer/Concert/ConcertApp/MultiUserClient/Source/MultiUserClient/Private/Widgets/ActiveSession/Replication/Client/Single/Columns/SingleClientColumns_ConflictWarning.cpp
--- a/Engine/Plugins/Developer/Concert/ConcertApp/MultiUserClient/Source/MultiUserClient/Private/Widgets/ActiveSession/Replication/Client/Single/Columns/SingleClientColumns_ConflictWarning.cpp
+++ b/Engine/Plugins/Developer/Concert/ConcertApp/MultiUserClient/Source/MultiUserClient/Private/Widgets/ActiveSession/Replication/Client/Single/Columns/SingleClientColumns_ConflictWarning.cpp
@@ -19,111 +19,173 @@
 namespace UE::MultiUserClient::SingleClientColumns
 {
 	const FName ConflictWarningTopLevelObjectColumnId = TEXT("ConflictWarningTopLevelObjectColumn");
-	const FName ConflictWarningSubobjectColumnId = TEXT("ConflictWarningSubobjectColumn");
 	const FName ConflictWarningPropertyColumnId = TEXT("ConflictWarningPropertyColumn");
 	
-	ConcertSharedSlate::ReplicationColumns::FReplicationTopLevelObjectColumn ConflictWarningForObject(
+	ConcertSharedSlate::FObjectColumnEntry ConflictWarningForObject(
 		TSharedRef<IConcertClient> InClient,
 		FGlobalAuthorityCache& InAuthorityCache,
 		const FGuid& ClientId
 		)
 	{
-		using FColumnType = ConcertSharedSlate::ReplicationColumns::FReplicationTopLevelObjectColumn;
-		return FColumnType(
-			FColumnType::FArguments()
-				.GenerateWidgetColumn_Lambda([InClient = MoveTemp(InClient), &InAuthorityCache, ClientId](const FColumnType::FBuildArgs& InArgs) mutable
+		class FPropertyColumn_ConflictWarningForObject : public ConcertSharedSlate::IObjectTreeColumn
+		{
+		public:
+
+			FPropertyColumn_ConflictWarningForObject(
+				TSharedRef<IConcertClient> InClient,
+				FGlobalAuthorityCache& InAuthorityCache,
+				const FGuid& ClientId
+				)
+				: Client(MoveTemp(InClient))
+				, AuthorityCache(InAuthorityCache)
+				, ClientId(ClientId)
+			{}
+			
+			virtual SHeaderRow::FColumn::FArguments CreateHeaderRowArgs() const override
+			{
+				return SHeaderRow::Column(ConflictWarningTopLevelObjectColumnId)
+					.DefaultLabel(FText::GetEmpty())
+					.FixedWidth(20.f);
+			}
+			
+			virtual TSharedRef<SWidget> GenerateColumnWidget(const FBuildArgs& InArgs) override
+			{
+				const auto GetVisibility = [this, ObjectPath = InArgs.RowItem.RowData.GetObjectPath()]()
 				{
-					const auto GetVisibility = [&InAuthorityCache, ClientId, ObjectPath = InArgs.RowData.GetObjectPath()]()
-					{
-						const EAuthorityMutability TakeAuthorityResult = InAuthorityCache.CanClientTakeAuthorityAfterSubmission(ObjectPath, ClientId);
-						return TakeAuthorityResult == EAuthorityMutability::Conflict
-							? EVisibility::Visible
-							: EVisibility::Collapsed;
-					};
-					const auto GetToolTip = [InClient = MoveTemp(InClient), &InAuthorityCache, ClientId, ObjectPath = InArgs.RowData.GetObjectPath()]()
-					{
-						TSet<FString> ClientNames;
-						const EAuthorityMutability TakeAuthorityResult = InAuthorityCache.CanClientTakeAuthorityAfterSubmission(ObjectPath, ClientId,
-						[InClient, &ClientNames](const FGuid& ClientId, const FConcertPropertyChain& ConflictingProperty)
+					const EAuthorityMutability TakeAuthorityResult = AuthorityCache.CanClientTakeAuthorityAfterSubmission(ObjectPath, ClientId);
+					return TakeAuthorityResult == EAuthorityMutability::Conflict
+						? EVisibility::Visible
+						: EVisibility::Collapsed;
+				};
+				const auto GetToolTip = [this, ObjectPath = InArgs.RowItem.RowData.GetObjectPath()]()
+				{
+					TSet<FString> ClientNames;
+					const EAuthorityMutability TakeAuthorityResult = AuthorityCache.CanClientTakeAuthorityAfterSubmission(
+						ObjectPath,
+						ClientId,
+						[this, &ClientNames](const FGuid& InClientId, const FConcertPropertyChain&)
 						{
-							ClientNames.Add(ClientUtils::GetClientDisplayName(*InClient, ClientId));
+							ClientNames.Add(ClientUtils::GetClientDisplayName(*Client, InClientId));
 							return EBreakBehavior::Continue;
 						});
 						
-						return TakeAuthorityResult != EAuthorityMutability::Conflict
-							? FText::GetEmpty()
-							: FText::Format(
-								LOCTEXT("Subobject.WarningFmt", "{0} {1}|plural(one=is,other=are) replicating some of the assigned properties. This change will not be sent to the server."),
-								FText::FromString(FString::Join(ClientNames, TEXT(", "))),
-								ClientNames.Num()
-							);
-					};
+					return TakeAuthorityResult != EAuthorityMutability::Conflict
+						? FText::GetEmpty()
+						: FText::Format(
+							LOCTEXT("Subobject.WarningFmt", "{0} {1}|plural(one=is,other=are) replicating some of the assigned properties. This change will not be sent to the server."),
+							FText::FromString(FString::Join(ClientNames, TEXT(", "))),
+							ClientNames.Num()
+						);
+				};
 					
-					return SNew(SWarningIcon)
-						.Visibility_Lambda(GetVisibility)
-						.ToolTipText_Lambda(GetToolTip);
-				})
-				.PopulateSearchItems_Lambda([](const ConcertSharedSlate::FReplicatedObjectData&, TArray<FString>&){})
-				.ColumnSortOrder(static_cast<int32>(ETopLevelObjectColumnOrder::ConflictWarning)),
-			SHeaderRow::Column(ConflictWarningSubobjectColumnId)
-				.DefaultLabel(FText::GetEmpty())
-				.FixedWidth(20.f)
-			);
+				return SNew(SWarningIcon)
+					.Visibility_Lambda(GetVisibility)
+					.ToolTipText_Lambda(GetToolTip);
+			}
+
+		private:
+
+			const TSharedRef<IConcertClient> Client;
+			FGlobalAuthorityCache& AuthorityCache;
+			const FGuid ClientId;
+		};
+
+		return {
+			ConcertSharedSlate::TReplicationColumnDelegates<ConcertSharedSlate::FObjectTreeRowContext>::FCreateColumn::CreateLambda(
+				[InClient = MoveTemp(InClient), &InAuthorityCache, ClientId]()
+				{
+					return MakeShared<FPropertyColumn_ConflictWarningForObject>(InClient, InAuthorityCache, ClientId);
+				}),
+			ConflictWarningTopLevelObjectColumnId,
+			{ static_cast<int32>(ETopLevelObjectColumnOrder::ConflictWarning) }
+		};
 	}
 
-	ConcertSharedSlate::ReplicationColumns::FReplicationPropertyColumn ConflictWarningForProperty(
+	ConcertSharedSlate::FPropertyColumnEntry ConflictWarningForProperty(
 		TSharedRef<IConcertClient> InClient,
 		TAttribute<const ConcertSharedSlate::IReplicationStreamViewer*> Viewer,
 		FGlobalAuthorityCache& InAuthorityCache,
 		const FGuid& ClientId
 		)
 	{
-		using FColumnType = ConcertSharedSlate::ReplicationColumns::FReplicationPropertyColumn;
-		return FColumnType(
-			FColumnType::FArguments()
-				.GenerateWidgetColumn_Lambda([InClient = MoveTemp(InClient), Viewer, &InAuthorityCache, ClientId](const FColumnType::FBuildArgs& InArgs) mutable
+		class FPropertyColumn_ConflictWarningForProperty : public ConcertSharedSlate::IPropertyTreeColumn
+		{
+		public:
+
+			FPropertyColumn_ConflictWarningForProperty(
+				TSharedRef<IConcertClient> InClient,
+				TAttribute<const ConcertSharedSlate::IReplicationStreamViewer*> Viewer,
+				FGlobalAuthorityCache& InAuthorityCache,
+				const FGuid& ClientId
+				)
+				: Client(MoveTemp(InClient))
+				, Viewer((MoveTemp(Viewer)))
+				, AuthorityCache(InAuthorityCache)
+				, ClientId(ClientId)
+			{}
+			
+			virtual SHeaderRow::FColumn::FArguments CreateHeaderRowArgs() const override
+			{
+				return SHeaderRow::Column(ConflictWarningPropertyColumnId)
+					.DefaultLabel(FText::GetEmpty())
+					.FixedWidth(20.f);
+			}
+			
+			virtual TSharedRef<SWidget> GenerateColumnWidget(const FBuildArgs& InArgs) override
+			{
+				const auto GetVisibility = [this, Property = InArgs.RowItem.RowData.GetProperty()]()
 				{
-					const auto GetVisibility = [&Viewer, &InAuthorityCache, ClientId, Property = InArgs.RowData.GetProperty()]()
-					{
-						const bool bCanAddProperty = Algo::AllOf(Viewer.Get()->GetObjectsBeingPropertyEdited(),
-							[&InAuthorityCache, ClientId, Property](const FSoftObjectPath& ObjectPath)
-							{
-								return InAuthorityCache.CanClientAddProperty(ObjectPath, ClientId, Property);
-							});
-						return bCanAddProperty
-							? EVisibility::Collapsed
-							: EVisibility::Visible;
-					};
-					const auto GetToolTip = [InClient = MoveTemp(InClient), Viewer, &InAuthorityCache, ClientId, Property = InArgs.RowData.GetProperty()]()
+					const bool bCanAddProperty = Algo::AllOf(Viewer.Get()->GetObjectsBeingPropertyEdited(),
+						[this, Property](const FSoftObjectPath& ObjectPath)
+						{
+							return AuthorityCache.CanClientAddProperty(ObjectPath, ClientId, Property);
+						});
+					return bCanAddProperty
+						? EVisibility::Collapsed
+						: EVisibility::Visible;
+				};
+				const auto GetToolTip = [this, Property = InArgs.RowItem.RowData.GetProperty()]()
+				{
+					TSet<FString> Names;
+					for (const FSoftObjectPath& ObjectPath: Viewer.Get()->GetObjectsBeingPropertyEdited())
 					{
-						TSet<FString> Names;
-						for (const FSoftObjectPath& ObjectPath: Viewer.Get()->GetObjectsBeingPropertyEdited())
+						const TOptional<FGuid> Owner = AuthorityCache.GetClientWithAuthorityOverProperty(ObjectPath, Property);
+						if (Owner || *Owner != ClientId)
 						{
-							const TOptional<FGuid> Owner = InAuthorityCache.GetClientWithAuthorityOverProperty(ObjectPath, Property);
-							if (Owner || *Owner != ClientId)
-							{
-								Names.Add(ClientUtils::GetClientDisplayName(*InClient, *Owner));
-							}
+							Names.Add(ClientUtils::GetClientDisplayName(*Client, *Owner));
 						}
+					}
 						
-						return Names.IsEmpty()
-							? FText::GetEmpty()
-							: FText::Format(
-								LOCTEXT("Property.WarningFmt", "Replicated by {0}. This property will not be submitted to the server."),
-								FText::FromString(FString::Join(Names, TEXT(", ")))
-								);
-					};
+					return Names.IsEmpty()
+						? FText::GetEmpty()
+						: FText::Format(
+							LOCTEXT("Property.WarningFmt", "Replicated by {0}. This property will not be submitted to the server."),
+							FText::FromString(FString::Join(Names, TEXT(", ")))
+							);
+				};
 					
-					return SNew(SWarningIcon)
-						.Visibility_Lambda(GetVisibility)
-						.ToolTipText_Lambda(GetToolTip);
-				})
-				.PopulateSearchItems_Lambda([](const ConcertSharedSlate::FReplicatedPropertyData&, TArray<FString>&){})
-				.ColumnSortOrder(static_cast<int32>(EPropertyColumnOrder::ConflictWarning)),
-			SHeaderRow::Column(ConflictWarningPropertyColumnId)
-				.DefaultLabel(FText::GetEmpty())
-				.FixedWidth(20.f)
-			);
+				return SNew(SWarningIcon)
+					.Visibility_Lambda(GetVisibility)
+					.ToolTipText_Lambda(GetToolTip);
+			}
+
+		private:
+
+			const TSharedRef<IConcertClient> Client;
+			const TAttribute<const ConcertSharedSlate::IReplicationStreamViewer*> Viewer;
+			FGlobalAuthorityCache& AuthorityCache;
+			const FGuid ClientId;
+		};
+
+		return {
+			ConcertSharedSlate::TReplicationColumnDelegates<ConcertSharedSlate::FPropertyTreeRowContext>::FCreateColumn::CreateLambda(
+				[InClient = MoveTemp(InClient), Viewer = MoveTemp(Viewer), &InAuthorityCache, ClientId]()
+				{
+					return MakeShared<FPropertyColumn_ConflictWarningForProperty>(InClient, Viewer, InAuthorityCache, ClientId);
+				}),
+			ConflictWarningPropertyColumnId,
+			{ static_cast<int32>(EPropertyColumnOrder::ConflictWarning) }
+		};
 	}
 }
 
diff --git a/Engine/Plugins/Developer/Concert/ConcertApp/MultiUserClient/Source/MultiUserClient/Private/Widgets/ActiveSession/Replication/Client/Single/Columns/SingleClientColumns_Owner.cpp b/Engine/Plugins/Developer/Concert/ConcertApp/MultiUserClient/Source/MultiUserClient/Private/Widgets/ActiveSession/Replication/Client/Single/Columns/SingleClientColumns_Owner.cpp
--- a/Engine/Plugins/Developer/Concert/ConcertApp/MultiUserClient/Source/MultiUserClient/Private/Widgets/ActiveSession/Replication/Client/Single/Columns/SingleClientColumns_Owner.cpp
+++ b/Engine/Plugins/Developer/Concert/ConcertApp/MultiUserClient/Source/MultiUserClient/Private/Widgets/ActiveSession/Replication/Client/Single/Columns/SingleClientColumns_Owner.cpp
@@ -5,7 +5,7 @@
 #include "IConcertClient.h"
 #include "MultiUserReplicationStyle.h"
 #include "Replication/Editor/View/IReplicationStreamViewer.h"
-#include "Replication/Editor/View/ReplicationColumnsUtils.h"
+#include "Replication/Editor/View/Column/ReplicationColumnsUtils.h"
 #include "Replication/Util/GlobalAuthorityCache.h"
 #include "SOwnerClientList.h"
 #include "Widgets/ActiveSession/Replication/Client/ClientUtils.h"
@@ -28,37 +28,65 @@ namespace UE::MultiUserClient::SingleClientColumns
 		}
 	}
 	
-	ConcertSharedSlate::ReplicationColumns::FReplicationTopLevelObjectColumn OwnerOfObject(
-		const TSharedRef<IConcertClient>& InClient,
+	ConcertSharedSlate::FObjectColumnEntry OwnerOfObject(
+		TSharedRef<IConcertClient> InClient,
 		FGlobalAuthorityCache& InAuthorityCache
 		)
 	{
-		using FColumnType = ConcertSharedSlate::ReplicationColumns::FReplicationTopLevelObjectColumn;
-		return FColumnType(
-			FColumnType::FArguments()
-				.GenerateWidgetColumn_Lambda([InClient, &InAuthorityCache](const FColumnType::FBuildArgs& InArgs)
-				{
-					return SNew(SOwnerClientList, InClient, InAuthorityCache)
-						.GetClientList_Lambda([&InAuthorityCache, ObjectPath = InArgs.RowData.GetObjectPath()](const FGlobalAuthorityCache&)
-						{
-							return InAuthorityCache.GetClientsWithAuthorityOverObject(ObjectPath);
-						})
-						.HighlightText_Lambda([HighlightText = InArgs.HighlightText](){ return *HighlightText.Get(); });
-				})
-				.PopulateSearchItems_Lambda([InClient, &InAuthorityCache](const ConcertSharedSlate::FReplicatedObjectData& InArgs, TArray<FString>& InOutSearchStrings)
+		class FObjectColum_OwnerOfObject : public ConcertSharedSlate::IObjectTreeColumn
+		{
+		public:
+
+			FObjectColum_OwnerOfObject(
+				TSharedRef<IConcertClient> Client,
+				FGlobalAuthorityCache& AuthorityCache
+				)
+				: Client(MoveTemp(Client))
+				, AuthorityCache(AuthorityCache)
+			{}
+			
+			virtual SHeaderRow::FColumn::FArguments CreateHeaderRowArgs() const override
+			{
+				return SHeaderRow::Column(OwnerOfSubobjectColumnId)
+					.DefaultLabel(LOCTEXT("Owner.Label", "Assigned Clients"))
+					.ToolTipText(LOCTEXT("Owner.ToolTip", "Clients that have registered properties for an object"))
+					.FillSized(FMultiUserReplicationStyle::Get()->GetFloat(TEXT("SingleClient.Object.OwnerColumnWidth")));
+			}
+			
+			virtual TSharedRef<SWidget> GenerateColumnWidget(const FBuildArgs& InArgs) override
+			{
+				return SNew(SOwnerClientList, Client, AuthorityCache)
+					.GetClientList_Lambda([this, ObjectPath = InArgs.RowItem.RowData.GetObjectPath()](const FGlobalAuthorityCache&)
+					{
+						return AuthorityCache.GetClientsWithAuthorityOverObject(ObjectPath);
+					})
+					.HighlightText_Lambda([HighlightText = InArgs.HighlightText](){ return *HighlightText.Get(); });
+			}
+			
+			virtual void PopulateSearchString(const ConcertSharedSlate::FObjectTreeRowContext& InItem, TArray<FString>& InOutSearchStrings) const override
+			{
+				Private::Shared::PopulateSearchTerms(
+					Client,
+					AuthorityCache.GetClientsWithAuthorityOverObject(InItem.RowData.GetObjectPath()),
+					InOutSearchStrings
+					);
+			}
+
+		private:
+			
+			const TSharedRef<IConcertClient> Client;
+			FGlobalAuthorityCache& AuthorityCache;
+		};
+
+		return {
+			ConcertSharedSlate::TReplicationColumnDelegates<ConcertSharedSlate::FObjectTreeRowContext>::FCreateColumn::CreateLambda(
+				[InClient = MoveTemp(InClient), &InAuthorityCache]()
 				{
-					Private::Shared::PopulateSearchTerms(
-						InClient,
-						InAuthorityCache.GetClientsWithAuthorityOverObject(InArgs.GetObjectPath()),
-						InOutSearchStrings
-						);
-				})
-				.ColumnSortOrder(static_cast<int32>(ETopLevelObjectColumnOrder::Owner)),
-				SHeaderRow::Column(OwnerOfSubobjectColumnId)
-				.DefaultLabel(LOCTEXT("Owner.Label", "Assigned Clients"))
-				.ToolTipText(LOCTEXT("Owner.ToolTip", "Clients that have registered properties for an object"))
-				.FillSized(FMultiUserReplicationStyle::Get()->GetFloat(TEXT("SingleClient.Object.OwnerColumnWidth")))
-			);
+					return MakeShared<FObjectColum_OwnerOfObject>(InClient, InAuthorityCache);
+				}),
+			OwnerOfSubobjectColumnId,
+			{ static_cast<int32>(ETopLevelObjectColumnOrder::Owner) }
+		};
 	}
 	
 	namespace Private::Property
@@ -83,37 +111,68 @@ namespace UE::MultiUserClient::SingleClientColumns
 		}
 	}
 
-	ConcertSharedSlate::ReplicationColumns::FReplicationPropertyColumn OwnerOfProperty(
-		const TSharedRef<IConcertClient>& InClient,
+	ConcertSharedSlate::FPropertyColumnEntry OwnerOfProperty(
+		TSharedRef<IConcertClient> InClient,
 		FGlobalAuthorityCache& InAuthorityCache,
-		const TAttribute<const ConcertSharedSlate::IReplicationStreamViewer*>& InViewer
+		TAttribute<const ConcertSharedSlate::IReplicationStreamViewer*> InViewer
 		)
 	{
-		using FColumnType = ConcertSharedSlate::ReplicationColumns::FReplicationPropertyColumn;
-		return FColumnType(
-			FColumnType::FArguments()
-				.GenerateWidgetColumn_Lambda([InClient, &InAuthorityCache, InViewer](const FColumnType::FBuildArgs& InArgs)
-				{
-					return SNew(SOwnerClientList, InClient, InAuthorityCache)
-						.GetClientList_Lambda([&InAuthorityCache, InViewer, Property = InArgs.RowData.GetProperty()](const FGlobalAuthorityCache&)
-						{
-							return Private::Property::GetPropertyOwners(InAuthorityCache, *InViewer.Get(), Property);
-						})
-						.HighlightText_Lambda([HighlightText = InArgs.HighlightText](){ return *HighlightText.Get(); });
-				})
-				.PopulateSearchItems_Lambda([InClient, &InAuthorityCache, InViewer](const ConcertSharedSlate::FReplicatedPropertyData& InArgs, TArray<FString>& InOutSearchStrings)
+		class FPropertyColumn_OwnerOfProperty : public ConcertSharedSlate::IPropertyTreeColumn
+		{
+		public:
+			
+			FPropertyColumn_OwnerOfProperty(
+				TSharedRef<IConcertClient> Client,
+				FGlobalAuthorityCache& AuthorityCache,
+				TAttribute<const ConcertSharedSlate::IReplicationStreamViewer*> Viewer
+				)
+				: Client(MoveTemp(Client))
+				, AuthorityCache(AuthorityCache)
+				, Viewer(MoveTemp(Viewer))
+			{}
+			
+			virtual SHeaderRow::FColumn::FArguments CreateHeaderRowArgs() const override
+			{
+				return SHeaderRow::Column(OwnerOfPropertyColumnId)
+					.DefaultLabel(LOCTEXT("Property.Owner", "Owner"))
+					.FillSized(FMultiUserReplicationStyle::Get()->GetFloat(TEXT("SingleClient.Property.OwnerColumnWidth")));
+			}
+			
+			virtual TSharedRef<SWidget> GenerateColumnWidget(const FBuildArgs& InArgs) override
+			{
+				return SNew(SOwnerClientList, Client, AuthorityCache)
+					.GetClientList_Lambda([this, Property = InArgs.RowItem.RowData.GetProperty()](const FGlobalAuthorityCache&)
+					{
+						return Private::Property::GetPropertyOwners(AuthorityCache, *Viewer.Get(), Property);
+					})
+					.HighlightText_Lambda([HighlightText = InArgs.HighlightText](){ return *HighlightText.Get(); });
+			}
+			
+			virtual void PopulateSearchString(const ConcertSharedSlate::FPropertyTreeRowContext& InItem, TArray<FString>& InOutSearchStrings) const override
+			{
+				Private::Shared::PopulateSearchTerms(
+					Client,
+					Private::Property::GetPropertyOwners(AuthorityCache, *Viewer.Get(), InItem.RowData.GetProperty()),
+					InOutSearchStrings
+					);
+			}
+
+		private:
+			
+			const TSharedRef<IConcertClient> Client;
+			FGlobalAuthorityCache& AuthorityCache;
+			const TAttribute<const ConcertSharedSlate::IReplicationStreamViewer*> Viewer;
+		};
+
+		return {
+			ConcertSharedSlate::TReplicationColumnDelegates<ConcertSharedSlate::FPropertyTreeRowContext>::FCreateColumn::CreateLambda(
+				[InClient = MoveTemp(InClient), &InAuthorityCache, InViewer = MoveTemp(InViewer)]()
 				{
-					Private::Shared::PopulateSearchTerms(
-						InClient,
-						Private::Property::GetPropertyOwners(InAuthorityCache, *InViewer.Get(), InArgs.GetProperty()),
-						InOutSearchStrings
-						);
-				})
-				.ColumnSortOrder(static_cast<int32>(EPropertyColumnOrder::Owner)),
-			SHeaderRow::Column(OwnerOfPropertyColumnId)
-				.DefaultLabel(LOCTEXT("Property.Owner", "Owner"))
-				.FillSized(FMultiUserReplicationStyle::Get()->GetFloat(TEXT("SingleClient.Property.OwnerColumnWidth")))
-			);
+					return MakeShared<FPropertyColumn_OwnerOfProperty>(InClient, InAuthorityCache, InViewer);
+				}),
+			OwnerOfPropertyColumnId,
+			{ static_cast<int32>(EPropertyColumnOrder::Owner) }
+		};
 	}
 }
 
diff --git a/Engine/Plugins/Developer/Concert/ConcertApp/MultiUserClient/Source/MultiUserClient/Private/Widgets/ActiveSession/Replication/Client/Single/Columns/SingleClientColumns_ToggleAuthority.cpp b/Engine/Plugins/Developer/Concert/ConcertApp/MultiUserClient/Source/MultiUserClient/Private/Widgets/ActiveSession/Replication/Client/Single/Columns/SingleClientColumns_ToggleAuthority.cpp
--- a/Engine/Plugins/Developer/Concert/ConcertApp/MultiUserClient/Source/MultiUserClient/Private/Widgets/ActiveSession/Replication/Client/Single/Columns/SingleClientColumns_ToggleAuthority.cpp
+++ b/Engine/Plugins/Developer/Concert/ConcertApp/MultiUserClient/Source/MultiUserClient/Private/Widgets/ActiveSession/Replication/Client/Single/Columns/SingleClientColumns_ToggleAuthority.cpp
@@ -9,7 +9,9 @@
 #include "Replication/Authority/IClientAuthoritySynchronizer.h"
 #include "Replication/Editor/Model/IReplicationStreamModel.h"
 #include "Replication/Editor/View/IReplicationStreamViewer.h"
-#include "Replication/Editor/View/ReplicationColumnsUtils.h"
+#include "Replication/Editor/View/Column/ReplicationColumnsUtils.h"
+#include "Replication/Editor/View/Column/IObjectTreeColumn.h"
+#include "Replication/Editor/View/Column/IPropertyTreeColumn.h"
 #include "Replication/Submission/ISubmissionWorkflow.h"
 #include "Replication/Util/GlobalAuthorityCache.h"
 #include "Widgets/ClientName/SClientName.h"
@@ -18,37 +20,37 @@
 
 namespace UE::MultiUserClient::SingleClientColumns
 {
-	const FName ToggleTopLevelAuthorityColumnId = TEXT("ToggleTopLevelAuthorityColumn");
+	const FName ToggleObjectAuthorityColumnId = TEXT("ToggleTopLevelAuthorityColumn");
 
-	ConcertSharedSlate::ReplicationColumns::FReplicationTopLevelObjectColumn ToggleObjectAuthority(
+	ConcertSharedSlate::FObjectColumnEntry ToggleObjectAuthority(
 		FAuthorityChangeTracker& ChangeTracker,
 		ISubmissionWorkflow& SubmissionWorkflow
 		)
 	{
 		using namespace ConcertSharedSlate;
-		using FColumnDelegates = TReplicationColumnDelegates<FReplicatedObjectData>;
-		return MakeCheckboxColumn<FReplicatedObjectData>(
-			ToggleTopLevelAuthorityColumnId,
+		using FColumnDelegates = TCheckboxColumnDelegates<FObjectTreeRowContext>;
+		return MakeCheckboxColumn<FObjectTreeRowContext>(
+			ToggleObjectAuthorityColumnId,
 				FColumnDelegates(
 					FColumnDelegates::FGetColumnCheckboxState::CreateLambda(
-					[&ChangeTracker](const FReplicatedObjectData& ObjectData)
+					[&ChangeTracker](const FObjectTreeRowContext& ObjectData)
 					{
-						const bool bHasAuthority = ChangeTracker.GetAuthorityStateAfterApplied(ObjectData.GetObjectPath());
+						const bool bHasAuthority = ChangeTracker.GetAuthorityStateAfterApplied(ObjectData.RowData.GetObjectPath());
 						return bHasAuthority ? ECheckBoxState::Checked : ECheckBoxState::Unchecked;
 					}),
 					FColumnDelegates::FOnColumnCheckboxChanged::CreateLambda(
-					[&ChangeTracker](bool bIsChecked, const FReplicatedObjectData& ObjectData)
+					[&ChangeTracker](bool bIsChecked, const FObjectTreeRowContext& ObjectData)
 					{
-						ChangeTracker.SetAuthorityIfAllowed({ ObjectData.GetObjectPath() }, bIsChecked);
+						ChangeTracker.SetAuthorityIfAllowed({ ObjectData.RowData.GetObjectPath() }, bIsChecked);
 					}),
-					FColumnDelegates::FGetToolTipText::CreateLambda([&ChangeTracker, &SubmissionWorkflow](const FReplicatedObjectData& ObjectData)
+					FColumnDelegates::FGetToolTipText::CreateLambda([&ChangeTracker, &SubmissionWorkflow](const FObjectTreeRowContext& ObjectData)
 					{
 						if (!CanEverSubmit(SubmissionWorkflow.GetUploadability()))
 						{
 							return LOCTEXT("ToggleAuthority.ToolTip.NotSupported", "This client cannot be remotely edited.");
 						}
 						
-						switch (ChangeTracker.GetChangeAuthorityMutability(ObjectData.GetObjectPath()))
+						switch (ChangeTracker.GetChangeAuthorityMutability(ObjectData.RowData.GetObjectPath()))
 						{
 						case EAuthorityMutability::Allowed: return LOCTEXT("ToggleAuthority.ToolTip.Allowed", "Whether to replicate this object.");
 						case EAuthorityMutability::NotApplicable: return LOCTEXT("ToggleAuthority.ToolTip.NotApplicable", "Assign properties to replicate first.");
@@ -56,9 +58,9 @@ namespace UE::MultiUserClient::SingleClientColumns
 						default: checkNoEntry(); return FText::GetEmpty();
 						}
 					}),
-					FColumnDelegates::FIsEnabled::CreateLambda([&ChangeTracker](const FReplicatedObjectData& ObjectData)
+					FColumnDelegates::FIsEnabled::CreateLambda([&ChangeTracker](const FObjectTreeRowContext& ObjectData)
 					{
-						return ChangeTracker.CanSetAuthorityFor(ObjectData.GetObjectPath());
+						return ChangeTracker.CanSetAuthorityFor(ObjectData.RowData.GetObjectPath());
 					})
 				),
 			FText::GetEmpty(),
diff --git a/Engine/Plugins/Developer/Concert/ConcertApp/MultiUserClient/Source/MultiUserClient/Private/Widgets/ActiveSession/Replication/Client/Single/SReplicationClientView.cpp b/Engine/Plugins/Developer/Concert/ConcertApp/MultiUserClient/Source/MultiUserClient/Private/Widgets/ActiveSession/Replication/Client/Single/SReplicationClientView.cpp
--- a/Engine/Plugins/Developer/Concert/ConcertApp/MultiUserClient/Source/MultiUserClient/Private/Widgets/ActiveSession/Replication/Client/Single/SReplicationClientView.cpp
+++ b/Engine/Plugins/Developer/Concert/ConcertApp/MultiUserClient/Source/MultiUserClient/Private/Widgets/ActiveSession/Replication/Client/Single/SReplicationClientView.cpp
@@ -79,20 +79,20 @@ namespace UE::MultiUserClient
 				.IsEditingEnabled = TAttribute<bool>::CreateLambda([&SubmissionWorkflow](){ return CanEverSubmit(SubmissionWorkflow.GetUploadability()); }),
 				.EditingDisabledToolTipText = LOCTEXT("Editing.NotImplemented", "Editing remote clients is not implemented. You can only edit the local client."),
 			},
-			.AdditionalPropertyColumns =
+			.PropertyColumns =
 			{
-				SingleClientColumns::OwnerOfProperty(ConcertClient.ToSharedRef(), AuthorityCache, GetReplicationViewerAttribute),
-				SingleClientColumns::ConflictWarningForProperty(ConcertClient.ToSharedRef(), GetReplicationViewerAttribute, AuthorityCache, InReplicationClient.GetEndpointId())
+				SingleClientColumns::ConflictWarningForProperty(ConcertClient.ToSharedRef(), GetReplicationViewerAttribute, AuthorityCache, InReplicationClient.GetEndpointId()),
+				SingleClientColumns::OwnerOfProperty(ConcertClient.ToSharedRef(), AuthorityCache, GetReplicationViewerAttribute) 
 			},
 			.ObjectHierarchy = ConcertClientSharedSlate::CreateObjectHierarchyForComponentHierarchy(), // This makes actors have children in the top view
             .NameModel = ConcertClientSharedSlate::CreateEditorObjectNameModel(), // This makes actors use their labels, and components use the names given in the BP editor
             .OnExtendObjectsContextMenu = FExtendObjectMenu::CreateSP(this, &SReplicationClientView::ExtendObjectContextMenu),
-            .AdditionalObjectColumns = 
-            {
-            	SingleClientColumns::ToggleObjectAuthority(AuthorityTracker, SubmissionWorkflow),
-            	SingleClientColumns::ConflictWarningForObject(ConcertClient.ToSharedRef(), AuthorityCache, InReplicationClient.GetEndpointId()),
-            	SingleClientColumns::OwnerOfObject(ConcertClient.ToSharedRef(), AuthorityCache)
-            },
+			.ObjectColumns =
+			{
+				SingleClientColumns::ToggleObjectAuthority(AuthorityTracker, SubmissionWorkflow),
+				SingleClientColumns::OwnerOfObject(ConcertClient.ToSharedRef(), AuthorityCache),
+				SingleClientColumns::ConflictWarningForObject(ConcertClient.ToSharedRef(), AuthorityCache, InReplicationClient.GetEndpointId())
+			}
 		};
 
 		EditorView = ConcertClientSharedSlate::CreateDefaultStreamEditor(MoveTemp(DefaultEditorParams));
diff --git a/Engine/Plugins/Developer/Concert/ConcertScripting/ConcertReplicationScripting/Source/ConcertReplicationScriptingEditor/Private/Widgets/PropertySelectionColumn.cpp b/Engine/Plugins/Developer/Concert/ConcertScripting/ConcertReplicationScripting/Source/ConcertReplicationScriptingEditor/Private/Widgets/PropertySelectionColumn.cpp
--- a/Engine/Plugins/Developer/Concert/ConcertScripting/ConcertReplicationScripting/Source/ConcertReplicationScriptingEditor/Private/Widgets/PropertySelectionColumn.cpp
+++ b/Engine/Plugins/Developer/Concert/ConcertScripting/ConcertReplicationScripting/Source/ConcertReplicationScriptingEditor/Private/Widgets/PropertySelectionColumn.cpp
@@ -1,9 +1,9 @@
 // Copyright Epic Games, Inc. All Rights Reserved.
 
 #include "PropertySelectionColumn.h"
+#include "Replication/Editor/View/Column/ReplicationColumnsUtils.h"
 
 #include "Algo/AnyOf.h"
-#include "Replication/Editor/View/ReplicationColumnsUtils.h"
 
 #define LOCTEXT_NAMESPACE "PropertySelectionCheckboxColumn"
 
@@ -11,7 +11,7 @@ namespace UE::ConcertReplicationScriptingEditor
 {
 	const FName PropertySelectionCheckboxColumnId(TEXT("AddPropertyCheckboxColumn"));
 	
-	ConcertSharedSlate::ReplicationColumns::FReplicationPropertyColumn MakePropertySelectionCheckboxColumn(
+	ConcertSharedSlate::FPropertyColumnEntry MakePropertySelectionCheckboxColumn(
 		const TSet<FConcertPropertyChain>& SelectedProperties,
 		FOnSelectProperty OnSelectPropertyDelegate,
 		bool bIsEditable,
@@ -19,14 +19,14 @@ namespace UE::ConcertReplicationScriptingEditor
 		)
 	{
 		using namespace ConcertSharedSlate;
-		using FPropertyColumnDelegates = TReplicationColumnDelegates<FReplicatedPropertyData>;
-		return MakeCheckboxColumn<FReplicatedPropertyData>(
+		using FPropertyColumnDelegates = TCheckboxColumnDelegates<FPropertyTreeRowContext>;
+		return MakeCheckboxColumn<FPropertyTreeRowContext>(
 			PropertySelectionCheckboxColumnId,
 			FPropertyColumnDelegates(
 				FPropertyColumnDelegates::FGetColumnCheckboxState::CreateLambda(
-				[&SelectedProperties](const FReplicatedPropertyData& Data)
+				[&SelectedProperties](const FPropertyTreeRowContext& Data)
 				{
-					const FConcertPropertyChain& DisplayProperty = Data.GetProperty();
+					const FConcertPropertyChain& DisplayProperty = Data.RowData.GetProperty();
 					const bool bShouldBeSelected = SelectedProperties.Contains(DisplayProperty)
 						|| Algo::AnyOf(SelectedProperties, [&DisplayProperty](const FConcertPropertyChain& SelectedProperty)
 						{
@@ -35,16 +35,16 @@ namespace UE::ConcertReplicationScriptingEditor
 					return bShouldBeSelected ? ECheckBoxState::Checked : ECheckBoxState::Unchecked;
 				}),
 				FPropertyColumnDelegates::FOnColumnCheckboxChanged::CreateLambda(
-				[OnSelectPropertyDelegate = MoveTemp(OnSelectPropertyDelegate)](bool bIsChecked, const FReplicatedPropertyData& Data)
+				[OnSelectPropertyDelegate = MoveTemp(OnSelectPropertyDelegate)](bool bIsChecked, const FPropertyTreeRowContext& Data)
 				{
-					OnSelectPropertyDelegate.Execute(Data.GetProperty(), bIsChecked);
+					OnSelectPropertyDelegate.Execute(Data.RowData.GetProperty(), bIsChecked);
 				}),
 				FPropertyColumnDelegates::FGetToolTipText::CreateLambda(
-				[](const FReplicatedPropertyData& Data)
+				[](const FPropertyTreeRowContext&)
 				{
 					return LOCTEXT("IncludePropertyTooltip", "Whether the property is included");
 				}),
-				FPropertyColumnDelegates::FIsEnabled::CreateLambda([bIsEditable](const FReplicatedPropertyData& Data)
+				FPropertyColumnDelegates::FIsEnabled::CreateLambda([bIsEditable](const FPropertyTreeRowContext&)
 				{
 					return bIsEditable;
 				})),
diff --git a/Engine/Plugins/Developer/Concert/ConcertScripting/ConcertReplicationScripting/Source/ConcertReplicationScriptingEditor/Private/Widgets/PropertySelectionColumn.h b/Engine/Plugins/Developer/Concert/ConcertScripting/ConcertReplicationScripting/Source/ConcertReplicationScriptingEditor/Private/Widgets/PropertySelectionColumn.h
--- a/Engine/Plugins/Developer/Concert/ConcertScripting/ConcertReplicationScripting/Source/ConcertReplicationScriptingEditor/Private/Widgets/PropertySelectionColumn.h
+++ b/Engine/Plugins/Developer/Concert/ConcertScripting/ConcertReplicationScripting/Source/ConcertReplicationScriptingEditor/Private/Widgets/PropertySelectionColumn.h
@@ -2,7 +2,7 @@
 
 #pragma once
 
-#include "Replication/Editor/View/PredefinedReplicationColumns.h"
+#include "Replication/Editor/View/Column/IPropertyTreeColumn.h"
 #include "Delegates/Delegate.h"
 
 namespace UE::ConcertReplicationScriptingEditor
@@ -17,7 +17,7 @@ namespace UE::ConcertReplicationScriptingEditor
 	 * @param SelectedProperties The checkbox is checked if the path is in SelectedProperties.
 	 * @param OnSelectPropertyDelegate Executed when a checkbox is checked or unchecked.
 	 */
-	ConcertSharedSlate::ReplicationColumns::FReplicationPropertyColumn MakePropertySelectionCheckboxColumn(
+	ConcertSharedSlate::FPropertyColumnEntry MakePropertySelectionCheckboxColumn(
 		const TSet<FConcertPropertyChain>& SelectedProperties,
 		FOnSelectProperty OnSelectPropertyDelegate,
 		bool bIsEditable,
diff --git a/Engine/Plugins/Developer/Concert/ConcertScripting/ConcertReplicationScripting/Source/ConcertReplicationScriptingEditor/Private/Widgets/SConcertPropertyChainPicker.cpp b/Engine/Plugins/Developer/Concert/ConcertScripting/ConcertReplicationScripting/Source/ConcertReplicationScriptingEditor/Private/Widgets/SConcertPropertyChainPicker.cpp
--- a/Engine/Plugins/Developer/Concert/ConcertScripting/ConcertReplicationScripting/Source/ConcertReplicationScriptingEditor/Private/Widgets/SConcertPropertyChainPicker.cpp
+++ b/Engine/Plugins/Developer/Concert/ConcertScripting/ConcertReplicationScripting/Source/ConcertReplicationScriptingEditor/Private/Widgets/SConcertPropertyChainPicker.cpp
@@ -92,10 +92,7 @@ namespace UE::ConcertReplicationScriptingEditor
 	{
 		ConcertSharedSlate::FCreatePropertyTreeViewParams Params
 		{
-			.PropertyColumns =
-			{
-				ConcertSharedSlate::ReplicationColumns::Property::LabelColumn()
-			},
+			.PropertyColumns ={ ConcertSharedSlate::ReplicationColumns::Property::LabelColumn()}
 		};
 		Params.NoItemsContent.Widget = SNew(STextBlock).Text(LOCTEXT("NoClass", "Select a class"));
 
@@ -103,7 +100,11 @@ namespace UE::ConcertReplicationScriptingEditor
 		{
 			Params.PrimaryPropertySort = { PropertySelectionCheckboxColumnId, EColumnSortMode::Ascending };
 			Params.PropertyColumns.Add(
-				MakePropertySelectionCheckboxColumn(*ContainedProperties, FOnSelectProperty::CreateSP(this, &SConcertPropertyChainPicker::OnPropertySelected), InArgs._IsEditable)
+				MakePropertySelectionCheckboxColumn(
+					*ContainedProperties,
+					FOnSelectProperty::CreateSP(this, &SConcertPropertyChainPicker::OnPropertySelected),
+					InArgs._IsEditable
+					)
 				);
 		}
 		
diff --git a/Engine/Plugins/Developer/Concert/ConcertUI/ConcertClientSharedSlate/Source/ConcertClientSharedSlate/Private/Replication/ClientReplicationWidgetFactories.cpp b/Engine/Plugins/Developer/Concert/ConcertUI/ConcertClientSharedSlate/Source/ConcertClientSharedSlate/Private/Replication/ClientReplicationWidgetFactories.cpp
--- a/Engine/Plugins/Developer/Concert/ConcertUI/ConcertClientSharedSlate/Source/ConcertClientSharedSlate/Private/Replication/ClientReplicationWidgetFactories.cpp
+++ b/Engine/Plugins/Developer/Concert/ConcertUI/ConcertClientSharedSlate/Source/ConcertClientSharedSlate/Private/Replication/ClientReplicationWidgetFactories.cpp
@@ -9,8 +9,8 @@
 #include "Replication/Editor/Model/Object/EditorObjectHierarchyModel.h"
 #include "Replication/Editor/Model/Object/EditorObjectNameModel.h"
 #include "Replication/Editor/Model/ReplicationStreamObject.h"
+#include "Replication/Editor/View/Column/SelectionViewerColumns.h"
 #include "Replication/Editor/View/ObjectEditor/SDefaultReplicationStreamEditor.h"
-#include "Replication/Editor/View/SelectionViewerColumns.h"
 #include "Replication/Editor/Model/TransactionalReplicationStreamModel.h"
 
 #include "UObject/UObjectGlobals.h"
@@ -84,21 +84,21 @@ namespace UE::ConcertClientSharedSlate
 		};
 		TSharedRef<FEditorIndirection> Indirection = MakeShared<FEditorIndirection>();
 		
-		const FReplicationPropertyColumn ReplicatesColumn = ReplicationColumns::Property::ReplicatesColumns(
+		const FPropertyColumnEntry ReplicatesColumn = ReplicationColumns::Property::ReplicatesColumns(
 			TAttribute<IReplicationStreamViewer*>::CreateLambda([Indirection](){ return Indirection->Editor.Get(); }),
 			Params.BaseEditorParams.DataModel,
 			MoveTemp(ExtendPropertiesDelegate),
-			TReplicationColumnDelegates<FReplicatedPropertyData>::FIsEnabled::CreateLambda([IsEnabled = Params.BaseEditorParams.IsEditingEnabled](const FReplicatedPropertyData&)
+			TCheckboxColumnDelegates<FPropertyTreeRowContext>::FIsEnabled::CreateLambda([IsEnabled = Params.BaseEditorParams.IsEditingEnabled](const FPropertyTreeRowContext&)
 			{
 				return !IsEnabled.IsBound() || IsEnabled.Get();
 			}),
 			Params.BaseEditorParams.EditingDisabledToolTipText
 			);
 		
-		TArray<FReplicationPropertyColumn>& PropertyColumns = Params.AdditionalPropertyColumns;
-		const bool bHasType = PropertyColumns.ContainsByPredicate([](const FReplicationPropertyColumn& Column)
+		TArray<FPropertyColumnEntry>& PropertyColumns = Params.PropertyColumns;
+		const bool bHasType = PropertyColumns.ContainsByPredicate([](const FPropertyColumnEntry& Entry)
 		{
-			return Column.ColumnId == Property::TypeColumnId;
+			return Entry.ColumnId == Property::TypeColumnId;
 		});
 		if (!bHasType)
 		{
@@ -108,11 +108,11 @@ namespace UE::ConcertClientSharedSlate
 
 		FCreateViewerParams ViewerParams
 		{
-			.PropertyTreeView = CreateFilterablePropertyTreeView({ .PropertyColumns = PropertyColumns }),
+			.PropertyTreeView = CreateFilterablePropertyTreeView({ .AdditionalPropertyColumns = MoveTemp(PropertyColumns) }),
 			.ObjectHierarchy = MoveTemp(Params.ObjectHierarchy),
 			.NameModel = MoveTemp(Params.NameModel),
 			.OnExtendObjectsContextMenu = MoveTemp(Params.OnExtendObjectsContextMenu),
-			.AdditionalObjectColumns = MoveTemp(Params.AdditionalObjectColumns),
+			.ObjectColumns = MoveTemp(Params.ObjectColumns),
 			.PrimaryObjectSort = FColumnSortInfo{ TopLevel::LabelColumnId, EColumnSortMode::Ascending },
 			.SecondaryObjectSort = FColumnSortInfo{ TopLevel::LabelColumnId, EColumnSortMode::Ascending },
 		};
diff --git a/Engine/Plugins/Developer/Concert/ConcertUI/ConcertClientSharedSlate/Source/ConcertClientSharedSlate/Private/Replication/Editor/View/ClientEditorColumns.cpp b/Engine/Plugins/Developer/Concert/ConcertUI/ConcertClientSharedSlate/Source/ConcertClientSharedSlate/Private/Replication/Editor/View/ClientEditorColumns.cpp
--- a/Engine/Plugins/Developer/Concert/ConcertUI/ConcertClientSharedSlate/Source/ConcertClientSharedSlate/Private/Replication/Editor/View/ClientEditorColumns.cpp
+++ b/Engine/Plugins/Developer/Concert/ConcertUI/ConcertClientSharedSlate/Source/ConcertClientSharedSlate/Private/Replication/Editor/View/ClientEditorColumns.cpp
@@ -4,12 +4,11 @@
 
 #include "Replication/Editor/Model/IEditableReplicationStreamModel.h"
 #include "Replication/Editor/Model/ReplicatedPropertyData.h"
-#include "Replication/Editor/View/ReplicationColumnsUtils.h"
+#include "Replication/Editor/View/Column/ReplicationColumnsUtils.h"
+#include "Replication/Editor/View/IReplicationStreamViewer.h"
 #include "Replication/ObjectUtils.h"
 
 #include "Internationalization/Internationalization.h"
-#include "Replication/Editor/View/IReplicationStreamViewer.h"
-#include "Replication/Editor/View/PredefinedReplicationColumns.h"
 
 #define LOCTEXT_NAMESPACE "ReplicationPropertyColumns"
 
@@ -87,40 +86,40 @@ namespace UE::ConcertClientSharedSlate::ReplicationColumns::Property
 		}
 	}
 	
-	ConcertSharedSlate::ReplicationColumns::FReplicationPropertyColumn ReplicatesColumns(
+	ConcertSharedSlate::FPropertyColumnEntry ReplicatesColumns(
 		TAttribute<ConcertSharedSlate::IReplicationStreamViewer*> Viewer,
 		TWeakPtr<ConcertSharedSlate::IEditableReplicationStreamModel> Model,
 		FExtendProperties ExtendPropertiesDelegate,
-		ConcertSharedSlate::TReplicationColumnDelegates<ConcertSharedSlate::FReplicatedPropertyData>::FIsEnabled IsEnabledDelegate,
+		ConcertSharedSlate::TCheckboxColumnDelegates<ConcertSharedSlate::FPropertyTreeRowContext>::FIsEnabled IsEnabledDelegate,
 		TAttribute<FText> DisabledToolTipText,
 		const float ColumnWidth,
 		const int32 Priority
 		)
 	{
 		using namespace ConcertSharedSlate;
-		using FPropertyColumnDelegates = TReplicationColumnDelegates<FReplicatedPropertyData>;
-		return MakeCheckboxColumn<FReplicatedPropertyData>(
+		using FPropertyColumnDelegates = TCheckboxColumnDelegates<FPropertyTreeRowContext>;
+		return MakeCheckboxColumn<FPropertyTreeRowContext>(
 			ReplicatesColumnId,
 			FPropertyColumnDelegates(
 				FPropertyColumnDelegates::FGetColumnCheckboxState::CreateLambda(
-				[Viewer, Model](const FReplicatedPropertyData& Data)
+				[Viewer, Model](const FPropertyTreeRowContext& Data)
 				{
 					const IReplicationStreamViewer* ViewerPin = Viewer.Get();
 					const TSharedPtr<IEditableReplicationStreamModel> ModelPin = Model.Pin();
-					return ensure(ViewerPin && ModelPin) ? Private::OnGetPropertyCheckboxState(Data.GetProperty(), *ViewerPin, *ModelPin) : ECheckBoxState::Undetermined;
+					return ensure(ViewerPin && ModelPin) ? Private::OnGetPropertyCheckboxState(Data.RowData.GetProperty(), *ViewerPin, *ModelPin) : ECheckBoxState::Undetermined;
 				}),
 				FPropertyColumnDelegates::FOnColumnCheckboxChanged::CreateLambda(
-				[Viewer, Model, ExtendPropertiesDelegate = MoveTemp(ExtendPropertiesDelegate)](bool bIsChecked, const FReplicatedPropertyData& Data)
+				[Viewer, Model, ExtendPropertiesDelegate = MoveTemp(ExtendPropertiesDelegate)](bool bIsChecked, const FPropertyTreeRowContext& Data)
 				{
 					const IReplicationStreamViewer* ViewerPin = Viewer.Get();
 					const TSharedPtr<IEditableReplicationStreamModel> ModelPin = Model.Pin();
 					if (ensure(ViewerPin && ModelPin))
 					{
-						Private::OnPropertyCheckboxChanged(bIsChecked, Data.GetProperty(), *ViewerPin, *ModelPin, ExtendPropertiesDelegate);
+						Private::OnPropertyCheckboxChanged(bIsChecked, Data.RowData.GetProperty(), *ViewerPin, *ModelPin, ExtendPropertiesDelegate);
 					}
 				}),
 				FPropertyColumnDelegates::FGetToolTipText::CreateLambda(
-				[IsEnabledDelegate, DisabledToolTipText = MoveTemp(DisabledToolTipText)](const FReplicatedPropertyData& Data)
+				[IsEnabledDelegate, DisabledToolTipText = MoveTemp(DisabledToolTipText)](const FPropertyTreeRowContext& Data)
 				{
 					const bool bIsDisabled = IsEnabledDelegate.IsBound() && !IsEnabledDelegate.Execute(Data);
 					const bool bCanCall = DisabledToolTipText.IsBound() || DisabledToolTipText.IsSet();
diff --git a/Engine/Plugins/Developer/Concert/ConcertUI/ConcertClientSharedSlate/Source/ConcertClientSharedSlate/Private/Replication/Editor/View/ClientEditorColumns.h b/Engine/Plugins/Developer/Concert/ConcertUI/ConcertClientSharedSlate/Source/ConcertClientSharedSlate/Private/Replication/Editor/View/ClientEditorColumns.h
--- a/Engine/Plugins/Developer/Concert/ConcertUI/ConcertClientSharedSlate/Source/ConcertClientSharedSlate/Private/Replication/Editor/View/ClientEditorColumns.h
+++ b/Engine/Plugins/Developer/Concert/ConcertUI/ConcertClientSharedSlate/Source/ConcertClientSharedSlate/Private/Replication/Editor/View/ClientEditorColumns.h
@@ -3,8 +3,9 @@
 #pragma once
 
 #include "Replication/ClientReplicationWidgetDelegates.h"
-#include "Replication/Editor/View/PredefinedReplicationColumns.h"
-#include "Replication/Editor/View/ReplicationColumnsUtils.h"
+#include "Replication/Editor/View/Column/ReplicationColumnsUtils.h"
+#include "Replication/Editor/View/Column/IPropertyTreeColumn.h"
+
 #include "Templates/SharedPointer.h"
 
 namespace UE::ConcertSharedSlate
@@ -37,11 +38,11 @@ namespace UE::ConcertClientSharedSlate::ReplicationColumns::Property
 	 * @param ColumnWidth Width to use for the column
 	 * @param Priority Determines position of this columns relative to the others
 	 */
-	ConcertSharedSlate::ReplicationColumns::FReplicationPropertyColumn ReplicatesColumns(
-		TAttribute<ConcertSharedSlate::IReplicationStreamViewer*> Viewer,
+	ConcertSharedSlate::FPropertyColumnEntry ReplicatesColumns(
+		TAttribute<ConcertSharedSlate::IReplicationStreamViewer*> Viewer, // TODO DP Columns: Refactor this
 		TWeakPtr<ConcertSharedSlate::IEditableReplicationStreamModel> Model,
 		FExtendProperties ExtendPropertiesDelegate = {},
-		ConcertSharedSlate::TReplicationColumnDelegates<ConcertSharedSlate::FReplicatedPropertyData>::FIsEnabled IsEnabledDelegate = {},
+		ConcertSharedSlate::TCheckboxColumnDelegates<ConcertSharedSlate::FPropertyTreeRowContext>::FIsEnabled IsEnabledDelegate = {},
 		TAttribute<FText> DisabledToolTipText = {},
 		const float ColumnWidth = 20.f,
 		const int32 Priority = static_cast<int32>(EReplicationPropertyColumnOrder::ReplicatesCheckbox)
diff --git a/Engine/Plugins/Developer/Concert/ConcertUI/ConcertClientSharedSlate/Source/ConcertClientSharedSlate/Private/Replication/Editor/View/ObjectEditor/SDefaultReplicationStreamEditor.h b/Engine/Plugins/Developer/Concert/ConcertUI/ConcertClientSharedSlate/Source/ConcertClientSharedSlate/Private/Replication/Editor/View/ObjectEditor/SDefaultReplicationStreamEditor.h
--- a/Engine/Plugins/Developer/Concert/ConcertUI/ConcertClientSharedSlate/Source/ConcertClientSharedSlate/Private/Replication/Editor/View/ObjectEditor/SDefaultReplicationStreamEditor.h
+++ b/Engine/Plugins/Developer/Concert/ConcertUI/ConcertClientSharedSlate/Source/ConcertClientSharedSlate/Private/Replication/Editor/View/ObjectEditor/SDefaultReplicationStreamEditor.h
@@ -3,10 +3,6 @@
 #pragma once
 
 #include "Replication/Editor/View/IReplicationStreamEditor.h"
-#include "Replication/Editor/View/PredefinedReplicationColumns.h"
-#include "Replication/Editor/View/ReplicationColumnsUtils.h"
-#include "Replication/ReplicationWidgetDelegates.h"
-
 #include "Widgets/DeclarativeSyntaxSupport.h"
 
 class SHorizontalBox;
diff --git a/Engine/Plugins/Developer/Concert/ConcertUI/ConcertClientSharedSlate/Source/ConcertClientSharedSlate/Private/Replication/Editor/View/PropertyTree/SFilteredPropertyTreeView.cpp b/Engine/Plugins/Developer/Concert/ConcertUI/ConcertClientSharedSlate/Source/ConcertClientSharedSlate/Private/Replication/Editor/View/PropertyTree/SFilteredPropertyTreeView.cpp
--- a/Engine/Plugins/Developer/Concert/ConcertUI/ConcertClientSharedSlate/Source/ConcertClientSharedSlate/Private/Replication/Editor/View/PropertyTree/SFilteredPropertyTreeView.cpp
+++ b/Engine/Plugins/Developer/Concert/ConcertUI/ConcertClientSharedSlate/Source/ConcertClientSharedSlate/Private/Replication/Editor/View/PropertyTree/SFilteredPropertyTreeView.cpp
@@ -60,27 +60,12 @@ namespace UE::ConcertClientSharedSlate
 		using namespace UE::ConcertSharedSlate;
 		FCreatePropertyTreeViewParams TreeViewParams
 		{
-			.PropertyColumns = MoveTemp(Params.PropertyColumns),
+			.PropertyColumns = MoveTemp(Params.AdditionalPropertyColumns),
 			.FilterItem = FFilterPropertyData::CreateSP(this, &SFilteredPropertyTreeView::PassesFilters),
 			.PrimaryPropertySort = Params.PrimaryPropertySort,
 			.SecondaryPropertySort = Params.SecondaryPropertySort
 		};
-		TreeViewParams.LeftOfPropertySearchBar.Widget = 
-			SNew(SHorizontalBox)
-
-			// The combo button for selecting the property filters
-			+SHorizontalBox::Slot()
-			.AutoWidth()
-			[
-				SBasicFilterBar<const FReplicatedPropertyData&>::MakeAddFilterButton(FilterBar.ToSharedRef())
-			]
-
-			+SHorizontalBox::Slot()
-			.AutoWidth()
-			[
-				InArgs._LeftOfSearchBar.Widget
-			]
-		;
+		TreeViewParams.LeftOfPropertySearchBar.Widget = SBasicFilterBar<const FReplicatedPropertyData&>::MakeAddFilterButton(FilterBar.ToSharedRef());
 		TreeViewParams.RowBelowSearchBar.Widget = FilterBar.ToSharedRef();
 		TreeViewParams.NoItemsContent.Widget = SNew(STextBlock).Text(LOCTEXT("AllFitlered", "All properties filtered."));
 			
diff --git a/Engine/Plugins/Developer/Concert/ConcertUI/ConcertClientSharedSlate/Source/ConcertClientSharedSlate/Private/Replication/Editor/View/PropertyTree/SFilteredPropertyTreeView.h b/Engine/Plugins/Developer/Concert/ConcertUI/ConcertClientSharedSlate/Source/ConcertClientSharedSlate/Private/Replication/Editor/View/PropertyTree/SFilteredPropertyTreeView.h
--- a/Engine/Plugins/Developer/Concert/ConcertUI/ConcertClientSharedSlate/Source/ConcertClientSharedSlate/Private/Replication/Editor/View/PropertyTree/SFilteredPropertyTreeView.h
+++ b/Engine/Plugins/Developer/Concert/ConcertUI/ConcertClientSharedSlate/Source/ConcertClientSharedSlate/Private/Replication/Editor/View/PropertyTree/SFilteredPropertyTreeView.h
@@ -3,7 +3,7 @@
 #pragma once
 
 #include "Replication/Editor/View/IPropertyTreeView.h"
-#include "Replication/Editor/View/ReplicationColumnsUtils.h"
+#include "Replication/Editor/View/Column/ReplicationColumnsUtils.h"
 
 #include "Filters/SBasicFilterBar.h"
 #include "Widgets/DeclarativeSyntaxSupport.h"
@@ -31,22 +31,6 @@ namespace UE::ConcertClientSharedSlate
 		
 		SLATE_BEGIN_ARGS(SFilteredPropertyTreeView)
 		{}
-			/*************** Arguments inherited by SReplicationTreeView ***************/
-		
-			/** The columns this list should have */
-			SLATE_ARGUMENT(TArray<ConcertSharedSlate::TReplicationColumn<ConcertSharedSlate::FReplicatedPropertyData>>, Columns)
-			/** The name of the column that will have the SExpanderArrow for the tree view. */
-			SLATE_ARGUMENT(FName, ExpandableColumnLabel)
-			/** Initial primary sort to set. */
-			SLATE_ARGUMENT(ConcertSharedSlate::FColumnSortInfo, PrimarySort)
-			/** Initial secondary sort to set. */
-			SLATE_ARGUMENT(ConcertSharedSlate::FColumnSortInfo, SecondarySort)
-		
-			/** Optional widget to add to the left of the search bar. */
-			SLATE_NAMED_SLOT(FArguments, LeftOfSearchBar)
-			/** Optional widget to add to the left of the search bar. */
-			SLATE_NAMED_SLOT(FArguments, RightOfSearchBar)
-		
 		SLATE_END_ARGS()
 
 		void Construct(const FArguments& InArgs, FFilterablePropertyTreeViewParams Params);
diff --git a/Engine/Plugins/Developer/Concert/ConcertUI/ConcertClientSharedSlate/Source/ConcertClientSharedSlate/Public/Replication/ClientReplicationWidgetFactories.h b/Engine/Plugins/Developer/Concert/ConcertUI/ConcertClientSharedSlate/Source/ConcertClientSharedSlate/Public/Replication/ClientReplicationWidgetFactories.h
--- a/Engine/Plugins/Developer/Concert/ConcertUI/ConcertClientSharedSlate/Source/ConcertClientSharedSlate/Public/Replication/ClientReplicationWidgetFactories.h
+++ b/Engine/Plugins/Developer/Concert/ConcertUI/ConcertClientSharedSlate/Source/ConcertClientSharedSlate/Public/Replication/ClientReplicationWidgetFactories.h
@@ -2,7 +2,7 @@
 
 #pragma once
 
-#include "ClientReplicationWidgetDelegates.h"
+#include "Replication/ClientReplicationWidgetDelegates.h"
 #include "Replication/ReplicationWidgetFactories.h"
 #include "Templates/SharedPointer.h"
 
@@ -11,7 +11,6 @@ struct FConcertObjectReplicationMap;
 
 namespace UE::ConcertSharedSlate
 {
-	struct FCreatePropertyTreeViewParams;
 	class IPropertySelectionSourceModel;
 	class IPropertyTreeView;
 	class IEditableReplicationStreamModel;
@@ -19,7 +18,9 @@ namespace UE::ConcertSharedSlate
 	class IReplicationStreamEditor;
 	class IStreamExtender;
 	class IObjectHierarchyModel;
+	
 	struct FCreateEditorParams;
+	struct FCreatePropertyTreeViewParams;
 }
 
 namespace UE::ConcertClientSharedSlate
@@ -48,7 +49,7 @@ namespace UE::ConcertClientSharedSlate
 	struct FFilterablePropertyTreeViewParams
 	{
 		/** The columns the property view should have. The label column is always included. */
-		TArray<ConcertSharedSlate::ReplicationColumns::FReplicationPropertyColumn> PropertyColumns
+		TArray<ConcertSharedSlate::FPropertyColumnEntry> AdditionalPropertyColumns
 		{
 			ConcertSharedSlate::ReplicationColumns::Property::LabelColumn(),
 			ConcertSharedSlate::ReplicationColumns::Property::TypeColumn()
@@ -77,7 +78,11 @@ namespace UE::ConcertClientSharedSlate
 		 * Additional property columns to add to the default ones.
 		 * By default the default editor has the label, type, and replicates (checkbox) columns.
 		 */
-		TArray<ConcertSharedSlate::ReplicationColumns::FReplicationPropertyColumn> AdditionalPropertyColumns;
+		TArray<ConcertSharedSlate::FPropertyColumnEntry> PropertyColumns
+		{
+			ConcertSharedSlate::ReplicationColumns::Property::LabelColumn(),
+			ConcertSharedSlate::ReplicationColumns::Property::TypeColumn()
+		};
 		
 		/**
 		 * Optional. Determines the objects displayed as children to the top-level objects in the top section.
@@ -102,7 +107,7 @@ namespace UE::ConcertClientSharedSlate
 		FExtendProperties OnExtendAddedProperties;
 		
 		/** Optional. Additional columns to add to the object view */
-		TArray<ConcertSharedSlate::ReplicationColumns::FReplicationTopLevelObjectColumn> AdditionalObjectColumns;
+		TArray<ConcertSharedSlate::FObjectColumnEntry> ObjectColumns;
 	};
 	
 	/**
diff --git a/Engine/Plugins/Developer/Concert/ConcertUI/ConcertSharedSlate/Source/ConcertSharedSlate/Private/Replication/Editor/View/Column/ObjectColumnAdapter.h b/Engine/Plugins/Developer/Concert/ConcertUI/ConcertSharedSlate/Source/ConcertSharedSlate/Private/Replication/Editor/View/Column/ObjectColumnAdapter.h
new file mode 100644
--- /dev/null
+++ b/Engine/Plugins/Developer/Concert/ConcertUI/ConcertSharedSlate/Source/ConcertSharedSlate/Private/Replication/Editor/View/Column/ObjectColumnAdapter.h
@@ -0,0 +1,59 @@
+// Copyright Epic Games, Inc. All Rights Reserved.
+
+#pragma once
+
+#include "Replication/Editor/Model/ReplicatedObjectData.h"
+#include "Replication/Editor/View/Column/IObjectTreeColumn.h"
+#include "Replication/Editor/View/Column/IReplicationTreeColumn.h"
+
+namespace UE::ConcertSharedSlate
+{
+	/** Adapts an IObjectTreeColumn to IReplicationTreeColumn<FReplicatedPropertyData>. It simply passes additional info down to IPropertyTreeColumn. */
+	class FObjectColumnAdapter : public IReplicationTreeColumn<FReplicatedObjectData>
+	{
+	public:
+
+		static TArray<TReplicationColumnEntry<FReplicatedObjectData>> Transform(const TArray<FObjectColumnEntry>& Entries)
+		{
+			TArray<TReplicationColumnEntry<FReplicatedObjectData>> Result;
+			Algo::Transform(Entries, Result, [](const FObjectColumnEntry& Entry) -> TReplicationColumnEntry<FReplicatedObjectData>
+			{
+				return {
+					TReplicationColumnDelegates<FReplicatedObjectData>::FCreateColumn::CreateLambda([CreateDelegate = Entry.CreateColumn]()
+						{
+							return MakeShared<FObjectColumnAdapter>(CreateDelegate.Execute());
+						}),
+					Entry.ColumnId,
+					Entry.ColumnInfo
+				};
+			});
+			return Result;
+		}
+		
+		FObjectColumnAdapter(TSharedRef<IObjectTreeColumn> InAdaptedColumn)
+			: AdaptedColumn(MoveTemp(InAdaptedColumn))
+		{}
+		
+		virtual SHeaderRow::FColumn::FArguments CreateHeaderRowArgs() const override { return AdaptedColumn->CreateHeaderRowArgs(); }
+		virtual TSharedRef<SWidget> GenerateColumnWidget(const FBuildArgs& InArgs) override
+		{
+			return AdaptedColumn->GenerateColumnWidget({ InArgs.HighlightText, Transform(InArgs.RowItem) });
+		}
+		virtual void PopulateSearchString(const FReplicatedObjectData& InItem, TArray<FString>& InOutSearchStrings) const override
+		{
+			return AdaptedColumn->PopulateSearchString(Transform(InItem), InOutSearchStrings);
+		}
+
+		virtual bool CanBeSorted() const override { return AdaptedColumn->CanBeSorted(); }
+		virtual bool IsLessThan(const FReplicatedObjectData& Left, const FReplicatedObjectData& Right) const override { return AdaptedColumn->IsLessThan(Transform(Left), Transform(Right)); }
+
+	private:
+
+		TSharedRef<IObjectTreeColumn> AdaptedColumn;
+
+		static FObjectTreeRowContext Transform(FReplicatedObjectData Data)
+		{
+			return { MoveTemp(Data) };
+		}
+	};
+}
diff --git a/Engine/Plugins/Developer/Concert/ConcertUI/ConcertSharedSlate/Source/ConcertSharedSlate/Private/Replication/Editor/View/Column/PropertyColumnAdapter.h b/Engine/Plugins/Developer/Concert/ConcertUI/ConcertSharedSlate/Source/ConcertSharedSlate/Private/Replication/Editor/View/Column/PropertyColumnAdapter.h
new file mode 100644
--- /dev/null
+++ b/Engine/Plugins/Developer/Concert/ConcertUI/ConcertSharedSlate/Source/ConcertSharedSlate/Private/Replication/Editor/View/Column/PropertyColumnAdapter.h
@@ -0,0 +1,59 @@
+// Copyright Epic Games, Inc. All Rights Reserved.
+
+#pragma once
+
+#include "Replication/Editor/Model/ReplicatedPropertyData.h"
+#include "Replication/Editor/View/Column/IPropertyTreeColumn.h"
+#include "Replication/Editor/View/Column/IReplicationTreeColumn.h"
+
+namespace UE::ConcertSharedSlate
+{
+	/** Adapts an IPropertyTreeColumn to IReplicationTreeColumn<FReplicatedPropertyData>. It simply passes additional info down to IPropertyTreeColumn. */
+	class FPropertyColumnAdapter : public IReplicationTreeColumn<FReplicatedPropertyData>
+	{
+	public:
+
+		static TArray<TReplicationColumnEntry<FReplicatedPropertyData>> Transform(const TArray<FPropertyColumnEntry>& Entries)
+		{
+			TArray<TReplicationColumnEntry<FReplicatedPropertyData>> Result;
+			Algo::Transform(Entries, Result, [](const FPropertyColumnEntry& Entry) -> TReplicationColumnEntry<FReplicatedPropertyData>
+			{
+				return {
+					TReplicationColumnDelegates<FReplicatedPropertyData>::FCreateColumn::CreateLambda([CreateDelegate = Entry.CreateColumn]()
+					{
+						return MakeShared<FPropertyColumnAdapter>(CreateDelegate.Execute());
+					}),
+					Entry.ColumnId,
+					Entry.ColumnInfo
+				};
+			});
+			return Result;
+		}
+		
+		FPropertyColumnAdapter(TSharedRef<IPropertyTreeColumn> InAdaptedColumn)
+			: AdaptedColumn(MoveTemp(InAdaptedColumn))
+		{}
+		
+		virtual SHeaderRow::FColumn::FArguments CreateHeaderRowArgs() const override { return AdaptedColumn->CreateHeaderRowArgs(); }
+		virtual TSharedRef<SWidget> GenerateColumnWidget(const FBuildArgs& InArgs) override
+		{
+			return AdaptedColumn->GenerateColumnWidget({ InArgs.HighlightText, Transform(InArgs.RowItem) });
+		}
+		virtual void PopulateSearchString(const FReplicatedPropertyData& InItem, TArray<FString>& InOutSearchStrings) const override
+		{
+			return AdaptedColumn->PopulateSearchString(Transform(InItem), InOutSearchStrings);
+		}
+
+		virtual bool CanBeSorted() const override { return AdaptedColumn->CanBeSorted(); }
+		virtual bool IsLessThan(const FReplicatedPropertyData& Left, const FReplicatedPropertyData& Right) const override { return AdaptedColumn->IsLessThan(Transform(Left), Transform(Right)); }
+
+	private:
+
+		TSharedRef<IPropertyTreeColumn> AdaptedColumn;
+
+		static FPropertyTreeRowContext Transform(FReplicatedPropertyData Data)
+		{
+			return { MoveTemp(Data) };
+		}
+	};
+}
diff --git a/Engine/Plugins/Developer/Concert/ConcertUI/ConcertSharedSlate/Source/ConcertSharedSlate/Private/Replication/Editor/View/ObjectEditor/SBaseReplicationStreamEditor.cpp b/Engine/Plugins/Developer/Concert/ConcertUI/ConcertSharedSlate/Source/ConcertSharedSlate/Private/Replication/Editor/View/ObjectEditor/SBaseReplicationStreamEditor.cpp
--- a/Engine/Plugins/Developer/Concert/ConcertUI/ConcertSharedSlate/Source/ConcertSharedSlate/Private/Replication/Editor/View/ObjectEditor/SBaseReplicationStreamEditor.cpp
+++ b/Engine/Plugins/Developer/Concert/ConcertUI/ConcertSharedSlate/Source/ConcertSharedSlate/Private/Replication/Editor/View/ObjectEditor/SBaseReplicationStreamEditor.cpp
@@ -44,7 +44,7 @@ namespace UE::ConcertSharedSlate
 		[
 			SAssignNew(ReplicationViewer, SReplicationStreamViewer, PropertiesModelAdapter.ToSharedRef())
 				.PropertyTreeView(InArgs._PropertyTreeView)
-				.AdditionalObjectColumns(InArgs._AdditionalObjectColumns)
+				.ObjectColumns(InArgs._ObjectColumns)
 				.PrimaryObjectSort(InArgs._PrimaryObjectSort)
 				.SecondaryObjectSort(InArgs._SecondaryObjectSort)
 				.ObjectHierarchy(InArgs._ObjectHierarchy)
diff --git a/Engine/Plugins/Developer/Concert/ConcertUI/ConcertSharedSlate/Source/ConcertSharedSlate/Private/Replication/Editor/View/ObjectEditor/SBaseReplicationStreamEditor.h b/Engine/Plugins/Developer/Concert/ConcertUI/ConcertSharedSlate/Source/ConcertSharedSlate/Private/Replication/Editor/View/ObjectEditor/SBaseReplicationStreamEditor.h
--- a/Engine/Plugins/Developer/Concert/ConcertUI/ConcertSharedSlate/Source/ConcertSharedSlate/Private/Replication/Editor/View/ObjectEditor/SBaseReplicationStreamEditor.h
+++ b/Engine/Plugins/Developer/Concert/ConcertUI/ConcertSharedSlate/Source/ConcertSharedSlate/Private/Replication/Editor/View/ObjectEditor/SBaseReplicationStreamEditor.h
@@ -4,8 +4,8 @@
 
 #include "Model/Item/SourceModelBuilders.h"
 #include "Replication/Editor/View/IReplicationStreamEditor.h"
-#include "Replication/Editor/View/PredefinedReplicationColumns.h"
-#include "Replication/Editor/View/ReplicationColumnsUtils.h"
+#include "Replication/Editor/View/Column/IObjectTreeColumn.h"
+#include "Replication/Editor/View/Column/ReplicationColumnsUtils.h"
 #include "Replication/ReplicationWidgetDelegates.h"
 
 #include "Widgets/DeclarativeSyntaxSupport.h"
@@ -56,7 +56,7 @@ namespace UE::ConcertSharedSlate
 			SLATE_ARGUMENT(TSharedPtr<IPropertyTreeView>, PropertyTreeView)
 		
 			/** Additional columns to add to the object view */
-			SLATE_ARGUMENT(TArray<ReplicationColumns::FReplicationTopLevelObjectColumn>, AdditionalObjectColumns)
+			SLATE_ARGUMENT(TArray<FObjectColumnEntry>, ObjectColumns)
 			/** Initial primary sort to set. */
 			SLATE_ARGUMENT(FColumnSortInfo, PrimaryObjectSort)
 			/** Initial secondary sort to set. */
diff --git a/Engine/Plugins/Developer/Concert/ConcertUI/ConcertSharedSlate/Source/ConcertSharedSlate/Private/Replication/Editor/View/ObjectViewer/Property/SPropertyTreeView.cpp b/Engine/Plugins/Developer/Concert/ConcertUI/ConcertSharedSlate/Source/ConcertSharedSlate/Private/Replication/Editor/View/ObjectViewer/Property/SPropertyTreeView.cpp
--- a/Engine/Plugins/Developer/Concert/ConcertUI/ConcertSharedSlate/Source/ConcertSharedSlate/Private/Replication/Editor/View/ObjectViewer/Property/SPropertyTreeView.cpp
+++ b/Engine/Plugins/Developer/Concert/ConcertUI/ConcertSharedSlate/Source/ConcertSharedSlate/Private/Replication/Editor/View/ObjectViewer/Property/SPropertyTreeView.cpp
@@ -5,6 +5,7 @@
 #include "Replication/Editor/Model/ReplicatedPropertyData.h"
 
 #include "Algo/ForEach.h"
+#include "Replication/Editor/View/Column/PropertyColumnAdapter.h"
 #include "UObject/UnrealType.h"
 
 #define LOCTEXT_NAMESPACE "SReplicatedPropertiesView"
@@ -19,7 +20,7 @@ namespace UE::ConcertSharedSlate
 				.RootItemsSource(&RootPropertyRowData)
 				.OnGetChildren(this, &SPropertyTreeView::GetPropertyRowChildren)
 				.FilterItem(InArgs._FilterItem)
-				.Columns(InArgs._Columns)
+				.Columns(FPropertyColumnAdapter::Transform(InArgs._Columns))
 				.ExpandableColumnLabel(InArgs._ExpandableColumnLabel)
 				.PrimarySort(InArgs._PrimarySort)
 				.SecondarySort(InArgs._SecondarySort)
diff --git a/Engine/Plugins/Developer/Concert/ConcertUI/ConcertSharedSlate/Source/ConcertSharedSlate/Private/Replication/Editor/View/ObjectViewer/Property/SPropertyTreeView.h b/Engine/Plugins/Developer/Concert/ConcertUI/ConcertSharedSlate/Source/ConcertSharedSlate/Private/Replication/Editor/View/ObjectViewer/Property/SPropertyTreeView.h
--- a/Engine/Plugins/Developer/Concert/ConcertUI/ConcertSharedSlate/Source/ConcertSharedSlate/Private/Replication/Editor/View/ObjectViewer/Property/SPropertyTreeView.h
+++ b/Engine/Plugins/Developer/Concert/ConcertUI/ConcertSharedSlate/Source/ConcertSharedSlate/Private/Replication/Editor/View/ObjectViewer/Property/SPropertyTreeView.h
@@ -7,6 +7,7 @@
 #include "Replication/Editor/View/Tree/SReplicationTreeView.h"
 
 #include "Filters/SBasicFilterBar.h"
+#include "Replication/Editor/View/Column/IPropertyTreeColumn.h"
 #include "Widgets/DeclarativeSyntaxSupport.h"
 #include "Widgets/SCompoundWidget.h"
 
@@ -24,8 +25,6 @@ namespace UE::ConcertSharedSlate
 	{
 	public:
 		
-		using FFilterRef = TSharedRef<FFilterBase<TSharedPtr<FReplicatedPropertyData>>>;
-		
 		SLATE_BEGIN_ARGS(SPropertyTreeView)
 		{}
 			/*************** Arguments inherited by SReplicationTreeView ***************/
@@ -34,7 +33,7 @@ namespace UE::ConcertSharedSlate
 			SLATE_EVENT(SReplicationTreeView<FReplicatedPropertyData>::FCustomFilter, FilterItem)
 		
 			/** The columns this list should have */
-			SLATE_ARGUMENT(TArray<TReplicationColumn<FReplicatedPropertyData>>, Columns)
+			SLATE_ARGUMENT(TArray<FPropertyColumnEntry>, Columns)
 			/** The name of the column that will have the SExpanderArrow for the tree view. */
 			SLATE_ARGUMENT(FName, ExpandableColumnLabel)
 			/** Initial primary sort to set. */
diff --git a/Engine/Plugins/Developer/Concert/ConcertUI/ConcertSharedSlate/Source/ConcertSharedSlate/Private/Replication/Editor/View/ObjectViewer/SReplicatedPropertyView.h b/Engine/Plugins/Developer/Concert/ConcertUI/ConcertSharedSlate/Source/ConcertSharedSlate/Private/Replication/Editor/View/ObjectViewer/SReplicatedPropertyView.h
--- a/Engine/Plugins/Developer/Concert/ConcertUI/ConcertSharedSlate/Source/ConcertSharedSlate/Private/Replication/Editor/View/ObjectViewer/SReplicatedPropertyView.h
+++ b/Engine/Plugins/Developer/Concert/ConcertUI/ConcertSharedSlate/Source/ConcertSharedSlate/Private/Replication/Editor/View/ObjectViewer/SReplicatedPropertyView.h
@@ -3,7 +3,7 @@
 #pragma once
 
 #include "Replication/Editor/View/IPropertyTreeView.h"
-#include "Replication/Editor/View/SelectionViewerColumns.h"
+#include "Replication/Editor/View/Column/SelectionViewerColumns.h"
 
 #include "Misc/Optional.h"
 #include "Widgets/DeclarativeSyntaxSupport.h"
diff --git a/Engine/Plugins/Developer/Concert/ConcertUI/ConcertSharedSlate/Source/ConcertSharedSlate/Private/Replication/Editor/View/ObjectViewer/SReplicationStreamViewer.cpp b/Engine/Plugins/Developer/Concert/ConcertUI/ConcertSharedSlate/Source/ConcertSharedSlate/Private/Replication/Editor/View/ObjectViewer/SReplicationStreamViewer.cpp
--- a/Engine/Plugins/Developer/Concert/ConcertUI/ConcertSharedSlate/Source/ConcertSharedSlate/Private/Replication/Editor/View/ObjectViewer/SReplicationStreamViewer.cpp
+++ b/Engine/Plugins/Developer/Concert/ConcertUI/ConcertSharedSlate/Source/ConcertSharedSlate/Private/Replication/Editor/View/ObjectViewer/SReplicationStreamViewer.cpp
@@ -6,8 +6,9 @@
 #include "Replication/Editor/Model/IReplicationStreamModel.h"
 #include "Replication/Editor/Model/Object/IObjectHierarchyModel.h"
 #include "Replication/Editor/Model/ReplicatedObjectData.h"
+#include "Replication/Editor/View/Column/ObjectColumnAdapter.h"
+#include "Replication/Editor/View/Column/SelectionViewerColumns.h"
 #include "Replication/Editor/View/ObjectViewer/Property/SPropertyTreeView.h"
-#include "Replication/Editor/View/SelectionViewerColumns.h"
 #include "Replication/ObjectUtils.h"
 #include "SReplicatedPropertyView.h"
 
@@ -189,12 +190,9 @@ namespace UE::ConcertSharedSlate
 
 	TSharedRef<SWidget> SReplicationStreamViewer::CreateOutlinerSection(const FArguments& InArgs)
 	{
-		TArray Columns
-		{
-			ReplicationColumns::TopLevel::LabelColumn(PropertiesModel.ToSharedRef(), NameModel.Get()),
-			ReplicationColumns::TopLevel::TypeColumn(PropertiesModel.ToSharedRef())
-		};
-		Columns.Append(InArgs._AdditionalObjectColumns);
+		TArray<FObjectColumnEntry> Columns = InArgs._ObjectColumns;
+		Columns.Add(ReplicationColumns::TopLevel::LabelColumn(PropertiesModel.ToSharedRef(), NameModel.Get()));
+		Columns.Add(ReplicationColumns::TopLevel::TypeColumn(PropertiesModel.ToSharedRef()));
 		
 		const bool bHasNoOutlinerObjectsAttribute = InArgs._NoOutlinerObjects.IsBound() || InArgs._NoOutlinerObjects.IsSet(); 
 		const TAttribute<FText> NoObjectsAttribute = bHasNoOutlinerObjectsAttribute ? InArgs._NoOutlinerObjects : LOCTEXT("NoObjects", "No objects to display");
@@ -218,7 +216,7 @@ namespace UE::ConcertSharedSlate
 			{
 				RequestPropertyDataRefresh();
 			})
-			.Columns(Columns)
+			.Columns(FObjectColumnAdapter::Transform(MoveTemp(Columns)))
 			.ExpandableColumnLabel(ReplicationColumns::TopLevel::LabelColumnId)
 			.PrimarySort(PrimaryObjectSort)
 			.SecondarySort(SecondaryObjectSort)
diff --git a/Engine/Plugins/Developer/Concert/ConcertUI/ConcertSharedSlate/Source/ConcertSharedSlate/Private/Replication/Editor/View/ObjectViewer/SReplicationStreamViewer.h b/Engine/Plugins/Developer/Concert/ConcertUI/ConcertSharedSlate/Source/ConcertSharedSlate/Private/Replication/Editor/View/ObjectViewer/SReplicationStreamViewer.h
--- a/Engine/Plugins/Developer/Concert/ConcertUI/ConcertSharedSlate/Source/ConcertSharedSlate/Private/Replication/Editor/View/ObjectViewer/SReplicationStreamViewer.h
+++ b/Engine/Plugins/Developer/Concert/ConcertUI/ConcertSharedSlate/Source/ConcertSharedSlate/Private/Replication/Editor/View/ObjectViewer/SReplicationStreamViewer.h
@@ -6,7 +6,8 @@
 
 #include "Replication/Editor/Model/ReplicatedObjectData.h"
 #include "Replication/Editor/View/Tree/SReplicationTreeView.h"
-#include "Replication/Editor/View/SelectionViewerColumns.h"
+#include "Replication/Editor/View/Column/IObjectTreeColumn.h"
+#include "Replication/Editor/View/Column/SelectionViewerColumns.h"
 #include "SReplicatedPropertyView.h"
 #include "StreamViewerObjectViewOptions.h"
 
@@ -53,7 +54,8 @@ namespace UE::ConcertSharedSlate
 			SLATE_ARGUMENT(TSharedPtr<IPropertyTreeView>, PropertyTreeView)
 		
 			/** Additional columns to add to the object view */
-			SLATE_ARGUMENT(TArray<ReplicationColumns::FReplicationTopLevelObjectColumn>, AdditionalObjectColumns)
+			SLATE_ARGUMENT(TArray<FObjectColumnEntry>, ObjectColumns)
+		
 			/** Initial primary sort to set. */
 			SLATE_ARGUMENT(FColumnSortInfo, PrimaryObjectSort)
 			/** Initial secondary sort to set. */
diff --git a/Engine/Plugins/Developer/Concert/ConcertUI/ConcertSharedSlate/Source/ConcertSharedSlate/Private/Replication/Editor/View/SelectionViewerColumns.cpp b/Engine/Plugins/Developer/Concert/ConcertUI/ConcertSharedSlate/Source/ConcertSharedSlate/Private/Replication/Editor/View/SelectionViewerColumns.cpp
--- a/Engine/Plugins/Developer/Concert/ConcertUI/ConcertSharedSlate/Source/ConcertSharedSlate/Private/Replication/Editor/View/SelectionViewerColumns.cpp
+++ b/Engine/Plugins/Developer/Concert/ConcertUI/ConcertSharedSlate/Source/ConcertSharedSlate/Private/Replication/Editor/View/SelectionViewerColumns.cpp
@@ -1,17 +1,17 @@
 // Copyright Epic Games, Inc. All Rights Reserved.
 
-#include "Replication/Editor/View/SelectionViewerColumns.h"
+#include "Replication/Editor/View/Column/SelectionViewerColumns.h"
 
 #include "ConcertFrontendStyle.h"
 #include "Replication/Editor/Model/IEditableReplicationStreamModel.h"
 #include "Replication/Editor/Model/ReplicatedPropertyData.h"
 #include "Replication/Editor/Model/ReplicatedObjectData.h"
 #include "Replication/Editor/View/DisplayUtils.h"
-#include "Replication/Editor/View/ObjectEditor/SBaseReplicationStreamEditor.h"
-#include "Replication/Editor/View/ReplicationColumnsUtils.h"
+#include "Replication/Editor/View/Column/ReplicationColumnsUtils.h"
 #include "Replication/PropertyChainUtils.h"
 
 #include "Internationalization/Internationalization.h"
+#include "Replication/Editor/View/Column/IObjectTreeColumn.h"
 #include "Textures/SlateIcon.h"
 #include "Widgets/Images/SImage.h"
 #include "Widgets/SBoxPanel.h"
@@ -27,80 +27,135 @@ namespace UE::ConcertSharedSlate::ReplicationColumns::TopLevel
 	const FName LabelColumnId = TEXT("LabelColumn");
 	const FName TypeColumnId = TEXT("TypeColumn");
 	
-	FReplicationTopLevelObjectColumn LabelColumn(TSharedRef<IReplicationStreamModel> Model, IObjectNameModel* OptionalNameModel)
+	FObjectColumnEntry LabelColumn(TSharedRef<IReplicationStreamModel> Model, IObjectNameModel* OptionalNameModel)
 	{
-		auto GetDisplayText = [OptionalNameModel](const FReplicatedObjectData& ObjectData)
+		class FLabelColumn_Object : public IObjectTreeColumn
 		{
-			const FSoftObjectPath& ObjectPath = ObjectData.GetObjectPath();
-			return DisplayUtils::GetObjectDisplayText(ObjectPath, OptionalNameModel);
-		};
-		
-		return FReplicationTopLevelObjectColumn(
-			FReplicationTopLevelObjectColumn::FArguments()
-				.GenerateWidgetColumn_Lambda([Model = MoveTemp(Model), GetDisplayText](const FReplicationTopLevelObjectColumn::FBuildArgs& Args)
-				{
-					const FText Text = GetDisplayText(Args.RowData);
+		public:
+			
+			FLabelColumn_Object(TSharedRef<IReplicationStreamModel> Model, IObjectNameModel* OptionalNameModel)
+				: Model(MoveTemp(Model))
+				, OptionalNameModel(OptionalNameModel)
+			{}
+
+			virtual SHeaderRow::FColumn::FArguments CreateHeaderRowArgs() const override
+			{
+				return SHeaderRow::Column(LabelColumnId)
+					.DefaultLabel(LOCTEXT("LabelColumnLabel", "Label"))
+					.FillSized(FConcertFrontendStyle::Get()->GetFloat("Concert.Replication.Tree.Object.LabelRowWidth"));
+			}
+			
+			virtual TSharedRef<SWidget> GenerateColumnWidget(const FBuildArgs& InArgs) override
+			{
+				const FReplicatedObjectData& ObjectData = InArgs.RowItem.RowData;
+				const FText Text = GetDisplayText(ObjectData);
 					
-					return SNew(SHorizontalBox)
-						.ToolTipText(FText::FromString(Args.RowData.GetObjectPath().ToString()))
+				return SNew(SHorizontalBox)
+					.ToolTipText(FText::FromString(ObjectData.GetObjectPath().ToString()))
 					
-						+SHorizontalBox::Slot()
-						.AutoWidth()
-						.HAlign(HAlign_Center)
-						.VAlign(VAlign_Center)
-						[
-							SNew(SImage)
-							.Image(DisplayUtils::GetObjectIcon(*Model, Args.RowData.GetObjectPath()).GetOptionalIcon())
-						]
+					+SHorizontalBox::Slot()
+					.AutoWidth()
+					.HAlign(HAlign_Center)
+					.VAlign(VAlign_Center)
+					[
+						SNew(SImage)
+						.Image(DisplayUtils::GetObjectIcon(*Model, ObjectData.GetObjectPath()).GetOptionalIcon())
+					]
 					
-						+SHorizontalBox::Slot()
-						.HAlign(HAlign_Left)
-						.VAlign(VAlign_Center)
-						.Padding(6.f, 0.f, 0.f, 0.f)
-						[
-							SNew(STextBlock)
-							.HighlightText(TAttribute<FText>::CreateLambda([HighlightText = Args.HighlightText](){ return *HighlightText; }))
-							.Text(Text)
-						];
-				})
-				.PopulateSearchItems_Lambda([OptionalNameModel](const FReplicatedObjectData& ObjectData, TArray<FString>& InOutSearchStrings)
-				{
-					InOutSearchStrings.Add(DisplayUtils::GetObjectDisplayText(ObjectData.GetObjectPath(), OptionalNameModel).ToString());
-				})
-				.IsLessThan_Lambda([GetDisplayText](const FReplicatedObjectData& Left, const FReplicatedObjectData& Right)
-				{
-					return GetDisplayText(Left).ToString() < GetDisplayText(Right).ToString();
-				})
-				.ColumnSortOrder(static_cast<int32>(ETopLevelColumnOrder::Label)),
-			SHeaderRow::Column(LabelColumnId)
-				.DefaultLabel(LOCTEXT("LabelColumnLabel", "Label"))
-				.FillSized(FConcertFrontendStyle::Get()->GetFloat("Concert.Replication.Tree.Object.LabelRowWidth"))
-			);
+					+SHorizontalBox::Slot()
+					.HAlign(HAlign_Left)
+					.VAlign(VAlign_Center)
+					.Padding(6.f, 0.f, 0.f, 0.f)
+					[
+						SNew(STextBlock)
+						.HighlightText(TAttribute<FText>::CreateLambda([HighlightText = InArgs.HighlightText](){ return *HighlightText; }))
+						.Text(Text)
+					];
+			}
+			
+			virtual void PopulateSearchString(const FObjectTreeRowContext& InItem, TArray<FString>& InOutSearchStrings) const override
+			{
+				const FReplicatedObjectData& ObjectData = InItem.RowData;
+				InOutSearchStrings.Add(DisplayUtils::GetObjectDisplayText(ObjectData.GetObjectPath(), OptionalNameModel).ToString());
+			}
+			
+			virtual bool CanBeSorted() const override { return true; } 
+			virtual bool IsLessThan(const FObjectTreeRowContext& Left, const FObjectTreeRowContext& Right) const override
+			{
+				return GetDisplayText(Left.RowData).ToString() < GetDisplayText(Right.RowData).ToString();
+			}
+
+		private:
+			
+			const TSharedRef<IReplicationStreamModel> Model;
+			IObjectNameModel* const OptionalNameModel;
+			
+			FText GetDisplayText(const FReplicatedObjectData& ObjectData) const
+			{
+				const FSoftObjectPath& ObjectPath = ObjectData.GetObjectPath();
+				return DisplayUtils::GetObjectDisplayText(ObjectPath, OptionalNameModel);
+			}
+		};
+
+		return {
+			TReplicationColumnDelegates<FObjectTreeRowContext>::FCreateColumn::CreateLambda([Model = MoveTemp(Model), OptionalNameModel]()
+			{
+				return MakeShared<FLabelColumn_Object>(Model, OptionalNameModel);
+			}),
+			LabelColumnId,
+			{ static_cast<int32>(ETopLevelColumnOrder::Label) }
+		};
 	}
 	
-	FReplicationTopLevelObjectColumn TypeColumn(TSharedRef<IReplicationStreamModel> Model)
+	FObjectColumnEntry TypeColumn(TSharedRef<IReplicationStreamModel> Model)
 	{
-		return FReplicationTopLevelObjectColumn(
-			FReplicationTopLevelObjectColumn::FArguments()
-				.GenerateWidgetColumn_Lambda([Model](const FReplicationTopLevelObjectColumn::FBuildArgs& Args)
-				{
-					return SNew(STextBlock)
-						.HighlightText(TAttribute<FText>::CreateLambda([HighlightText = Args.HighlightText](){ return *HighlightText; }))
-						.Text(DisplayUtils::GetObjectTypeText(*Model, Args.RowData.GetObjectPath()));
-				})
-				.PopulateSearchItems_Lambda([Model](const FReplicatedObjectData& ObjectData, TArray<FString>& InOutSearchStrings)
-				{
-					InOutSearchStrings.Add(DisplayUtils::GetObjectTypeText(Model.Get(), ObjectData.GetObjectPath()).ToString());
-				})
-				.IsLessThan_Lambda([Model](const FReplicatedObjectData& Left, const FReplicatedObjectData& Right)
-				{
-					return DisplayUtils::GetObjectTypeText(*Model, Left.GetObjectPath()).ToString() < DisplayUtils::GetObjectTypeText(*Model, Right.GetObjectPath()).ToString();
-				})
-				.ColumnSortOrder(static_cast<int32>(ETopLevelColumnOrder::Type)),
-			SHeaderRow::Column(TypeColumnId)
-				.DefaultLabel(LOCTEXT("TypeColumnLabel", "Type"))
-				.FillWidth(1.f)
-			);
+		class FLabelColumn_Type : public IObjectTreeColumn
+		{
+		public:
+			
+			FLabelColumn_Type(TSharedRef<IReplicationStreamModel> Model)
+				: Model(MoveTemp(Model))
+			{}
+
+			virtual SHeaderRow::FColumn::FArguments CreateHeaderRowArgs() const override
+			{
+				return SHeaderRow::Column(TypeColumnId)
+					.DefaultLabel(LOCTEXT("TypeColumnLabel", "Type"))
+					.FillWidth(1.f);
+			}
+			
+			virtual TSharedRef<SWidget> GenerateColumnWidget(const FBuildArgs& InArgs) override
+			{
+				return SNew(STextBlock)
+					.HighlightText(TAttribute<FText>::CreateLambda([HighlightText = InArgs.HighlightText](){ return *HighlightText; }))
+					.Text(DisplayUtils::GetObjectTypeText(*Model, InArgs.RowItem.RowData.GetObjectPath()));
+			}
+			
+			virtual void PopulateSearchString(const FObjectTreeRowContext& InItem, TArray<FString>& InOutSearchStrings) const override
+			{
+				InOutSearchStrings.Add(DisplayUtils::GetObjectTypeText(Model.Get(), InItem.RowData.GetObjectPath()).ToString());
+			}
+			
+			virtual bool CanBeSorted() const override { return true; } 
+			virtual bool IsLessThan(const FObjectTreeRowContext& Left, const FObjectTreeRowContext& Right) const override
+			{
+				return DisplayUtils::GetObjectTypeText(*Model, Left.RowData.GetObjectPath()).ToString()
+					< DisplayUtils::GetObjectTypeText(*Model, Right.RowData.GetObjectPath()).ToString();
+			}
+
+		private:
+			
+			const TSharedRef<IReplicationStreamModel> Model;
+		};
+
+		return {
+			TReplicationColumnDelegates<FObjectTreeRowContext>::FCreateColumn::CreateLambda([Model = MoveTemp(Model)]()
+			{
+				return MakeShared<FLabelColumn_Type>(Model);
+			}),
+			TypeColumnId,
+			{ static_cast<int32>(ETopLevelColumnOrder::Type) }
+		};
 	}
 }
 
@@ -113,63 +168,102 @@ namespace UE::ConcertSharedSlate::ReplicationColumns::Property
 	const FName LabelColumnId = TEXT("LabelColumn");
 	const FName TypeColumnId = TEXT("TypeColumn");
 	
-	FReplicationPropertyColumn LabelColumn()
+	FPropertyColumnEntry LabelColumn()
 	{
-		return FReplicationPropertyColumn(
-			FReplicationPropertyColumn::FArguments()
-				.GenerateWidgetColumn_Lambda([](const FReplicationPropertyColumn::FBuildArgs& Args)
-				{
-					return SNew(STextBlock)
+		class FLabelColumn_Property : public IPropertyTreeColumn
+		{
+		public:
+
+			virtual SHeaderRow::FColumn::FArguments CreateHeaderRowArgs() const override
+			{
+				return SHeaderRow::Column(LabelColumnId)
+					.DefaultLabel(LOCTEXT("LabelColumnLabel", "Label"))
+					.FillSized(FConcertFrontendStyle::Get()->GetFloat("Concert.Replication.Tree.Property.LabelRowWidth"));
+			}
+			
+			virtual TSharedRef<SWidget> GenerateColumnWidget(const FBuildArgs& InArgs) override
+			{
+				const FReplicatedPropertyData& PropertyData = InArgs.RowItem.RowData;
+				return SNew(STextBlock)
 						.Font(FAppStyle::GetFontStyle("PropertyWindow.NormalFont"))
-						.HighlightText(TAttribute<FText>::CreateLambda([HighlightText = Args.HighlightText](){ return *HighlightText; }))
-						.Text(DisplayUtils::GetPropertyDisplayText(Args.RowData.GetProperty()));
-				})
-				.PopulateSearchItems_Lambda([](const FReplicatedPropertyData& ObjectData, TArray<FString>& InOutSearchStrings)
-				{
-					InOutSearchStrings.Add(DisplayUtils::GetPropertyDisplayText(ObjectData.GetProperty()).ToString());
-				})
-				.IsLessThan_Lambda([](const FReplicatedPropertyData& Left, const FReplicatedPropertyData& Right)
-				{
-					return DisplayUtils::GetPropertyDisplayString(Left.GetProperty()) < DisplayUtils::GetPropertyDisplayString(Right.GetProperty());
-				})
-				.ColumnSortOrder(static_cast<int32>(EReplicationPropertyColumnOrder::Label)),
-			SHeaderRow::Column(LabelColumnId)
-				.DefaultLabel(LOCTEXT("LabelColumnLabel", "Label"))
-				.FillSized(FConcertFrontendStyle::Get()->GetFloat("Concert.Replication.Tree.Property.LabelRowWidth"))
-			);
+						.HighlightText(TAttribute<FText>::CreateLambda([HighlightText = InArgs.HighlightText](){ return *HighlightText; }))
+						.Text(DisplayUtils::GetPropertyDisplayText(PropertyData.GetProperty()));
+			}
+			
+			virtual void PopulateSearchString(const FPropertyTreeRowContext& InItem, TArray<FString>& InOutSearchStrings) const override
+			{
+				InOutSearchStrings.Add(DisplayUtils::GetPropertyDisplayText(InItem.RowData.GetProperty()).ToString());
+			}
+			
+			virtual bool CanBeSorted() const override { return true; } 
+			virtual bool IsLessThan(const FPropertyTreeRowContext& Left, const FPropertyTreeRowContext& Right) const override
+			{
+				return DisplayUtils::GetPropertyDisplayString(Left.RowData.GetProperty())
+					< DisplayUtils::GetPropertyDisplayString(Right.RowData.GetProperty());
+			}
+		};
+		
+		return {
+			TReplicationColumnDelegates<FPropertyTreeRowContext>::FCreateColumn::CreateLambda([]()
+			{
+				return MakeShared<FLabelColumn_Property>();
+			}),
+			LabelColumnId,
+			{ static_cast<int32>(EReplicationPropertyColumnOrder::Label) }
+		};
 	}
 	
-	FReplicationPropertyColumn TypeColumn()
+	FPropertyColumnEntry TypeColumn()
 	{
-		static auto GetDisplayText = [](const FReplicatedPropertyData& Args)
+		class FTypeColumn_Property : public IPropertyTreeColumn
 		{
-			UClass* Class = Args.GetOwningClass().TryLoadClass<UObject>();
-			const FProperty* Property = Class ? ConcertSyncCore::PropertyChain::ResolveProperty(*Class, Args.GetProperty()) : nullptr;
-			return Property ? FText::FromString(Property->GetCPPType()) : LOCTEXT("Unknown", "Unknown");	
+		public:
+
+			virtual SHeaderRow::FColumn::FArguments CreateHeaderRowArgs() const override
+			{
+				return SHeaderRow::Column(TypeColumnId)
+					.DefaultLabel(LOCTEXT("TypeColumnLabel", "Type"))
+					.FillWidth(1.f);
+			}
+			
+			virtual TSharedRef<SWidget> GenerateColumnWidget(const FBuildArgs& InArgs) override
+			{
+				const FReplicatedPropertyData& PropertyData = InArgs.RowItem.RowData;
+				return SNew(STextBlock)
+						.Font(FAppStyle::GetFontStyle("PropertyWindow.NormalFont"))
+						.HighlightText(TAttribute<FText>::CreateLambda([HighlightText = InArgs.HighlightText](){ return *HighlightText; }))
+						.Text(GetDisplayText(InArgs.RowItem.RowData));
+			}
+			
+			virtual void PopulateSearchString(const FPropertyTreeRowContext& InItem, TArray<FString>& InOutSearchStrings) const override
+			{
+				InOutSearchStrings.Add(GetDisplayText(InItem.RowData).ToString());
+			}
+			
+			virtual bool CanBeSorted() const override { return true; } 
+			virtual bool IsLessThan(const FPropertyTreeRowContext& Left, const FPropertyTreeRowContext& Right) const override
+			{
+				return GetDisplayText(Left.RowData).ToString() < GetDisplayText(Right.RowData).ToString();
+			}
+
+		private:
+			
+			static FText GetDisplayText(const FReplicatedPropertyData& Args)
+			{
+				UClass* Class = Args.GetOwningClass().TryLoadClass<UObject>();
+				const FProperty* Property = Class ? ConcertSyncCore::PropertyChain::ResolveProperty(*Class, Args.GetProperty()) : nullptr;
+				return Property ? FText::FromString(Property->GetCPPType()) : LOCTEXT("Unknown", "Unknown");	
+			};
 		};
 		
-		return FReplicationPropertyColumn(
-			FReplicationPropertyColumn::FArguments()
-				.GenerateWidgetColumn_Lambda([](const FReplicationPropertyColumn::FBuildArgs& Args)
-				{
-					return SNew(STextBlock)
-						.Font(FAppStyle::GetFontStyle("PropertyWindow.NormalFont"))
-						.HighlightText(TAttribute<FText>::CreateLambda([HighlightText = Args.HighlightText](){ return *HighlightText; }))
-						.Text(GetDisplayText(Args.RowData));
-				})
-				.PopulateSearchItems_Lambda([](const FReplicatedPropertyData& ObjectData, TArray<FString>& InOutSearchStrings)
-				{
-					InOutSearchStrings.Add(GetDisplayText(ObjectData).ToString());
-				})
-				.IsLessThan_Lambda([](const FReplicatedPropertyData& Left, const FReplicatedPropertyData& Right)
-				{
-					return GetDisplayText(Left).ToString() < GetDisplayText(Right).ToString();
-				})
-				.ColumnSortOrder(static_cast<int32>(EReplicationPropertyColumnOrder::Type)),
-			SHeaderRow::Column(TypeColumnId)
-				.DefaultLabel(LOCTEXT("TypeColumnLabel", "Type"))
-				.FillWidth(1.f)
-			);
+		return {
+			TReplicationColumnDelegates<FPropertyTreeRowContext>::FCreateColumn::CreateLambda([]()
+			{
+				return MakeShared<FTypeColumn_Property>();
+			}),
+			TypeColumnId,
+			{ static_cast<int32>(EReplicationPropertyColumnOrder::Type) }
+		};
 	}
 }
 
diff --git a/Engine/Plugins/Developer/Concert/ConcertUI/ConcertSharedSlate/Source/ConcertSharedSlate/Private/Replication/Editor/View/Tree/SReplicationColumnRow.h b/Engine/Plugins/Developer/Concert/ConcertUI/ConcertSharedSlate/Source/ConcertSharedSlate/Private/Replication/Editor/View/Tree/SReplicationColumnRow.h
--- a/Engine/Plugins/Developer/Concert/ConcertUI/ConcertSharedSlate/Source/ConcertSharedSlate/Private/Replication/Editor/View/Tree/SReplicationColumnRow.h
+++ b/Engine/Plugins/Developer/Concert/ConcertUI/ConcertSharedSlate/Source/ConcertSharedSlate/Private/Replication/Editor/View/Tree/SReplicationColumnRow.h
@@ -3,7 +3,6 @@
 #pragma once
 
 #include "ConcertFrontendStyle.h"
-#include "Replication/Editor/View/ReplicationColumn.h"
 
 #include "Widgets/Layout/SBox.h"
 #include "Widgets/SNullWidget.h"
@@ -22,10 +21,8 @@ namespace UE::ConcertSharedSlate
 	class SReplicationColumnRow : public SMultiColumnTableRow<TSharedPtr<TListItemType>>
 	{
 	public:
-
-		using TColumType = TReplicationColumn<TListItemType>;
 		
-		DECLARE_DELEGATE_RetVal_OneParam(const TColumType*, FGetColumn,
+		DECLARE_DELEGATE_RetVal_OneParam(TSharedPtr<IReplicationTreeColumn<TListItemType>>, FGetColumn,
 			const FName& ColumnId
 			);
 		DECLARE_DELEGATE_RetVal_TwoParams(TSharedPtr<SWidget>, FOverrideColumnWidget, const FName& ColumnName, const TListItemType& RowData);
@@ -85,13 +82,13 @@ namespace UE::ConcertSharedSlate
 				return ColumnOverride.ToSharedRef();
 			}
 			
-			const TColumType* Column = ColumnGetterDelegate.Execute(ColumnName); ensure(Column);
+			TSharedPtr<IReplicationTreeColumn<TListItemType>> Column = ColumnGetterDelegate.Execute(ColumnName); ensure(Column);
 			if (!Column)
 			{
 				return SNullWidget::NullWidget;
 			}
 			
-			const TSharedRef<SWidget> ColumnWidget = Column->BuildColumnWidget({ HighlightText, *RowData.Get() });
+			const TSharedRef<SWidget> ColumnWidget = Column->GenerateColumnWidget({ HighlightText, *RowData.Get() });
 			const bool bNeedsExpanderArrow = ColumnName == ExpandableColumnLabel;
 			if (!bNeedsExpanderArrow)
 			{
diff --git a/Engine/Plugins/Developer/Concert/ConcertUI/ConcertSharedSlate/Source/ConcertSharedSlate/Private/Replication/Editor/View/Tree/SReplicationTreeView.h b/Engine/Plugins/Developer/Concert/ConcertUI/ConcertSharedSlate/Source/ConcertSharedSlate/Private/Replication/Editor/View/Tree/SReplicationTreeView.h
--- a/Engine/Plugins/Developer/Concert/ConcertUI/ConcertSharedSlate/Source/ConcertSharedSlate/Private/Replication/Editor/View/Tree/SReplicationTreeView.h
+++ b/Engine/Plugins/Developer/Concert/ConcertUI/ConcertSharedSlate/Source/ConcertSharedSlate/Private/Replication/Editor/View/Tree/SReplicationTreeView.h
@@ -2,8 +2,8 @@
 
 #pragma once
 
-#include "Replication/Editor/View/ReplicationColumn.h"
-#include "Replication/Editor/View/ReplicationColumnsUtils.h"
+#include "Replication/Editor/View/Column/ReplicationColumnInfo.h"
+#include "Replication/Editor/View/Column/ReplicationColumnsUtils.h"
 #include "SReplicationColumnRow.h"
 
 #include "Algo/RemoveIf.h"
@@ -77,7 +77,7 @@ namespace UE::ConcertSharedSlate
 			SLATE_EVENT(FIsSearchableItem, IsSearchableItem)
 			
 			/** The columns this list should have */
-			SLATE_ARGUMENT(TArray<TReplicationColumn<TItemType>>, Columns)
+			SLATE_ARGUMENT(TArray<TReplicationColumnEntry<TItemType>>, Columns)
 			/** The name of the column that will have the SExpanderArrow for the tree view. */
 			SLATE_ARGUMENT(FName, ExpandableColumnLabel)
 			/** Visibility of the header row */
@@ -219,7 +219,7 @@ namespace UE::ConcertSharedSlate
 
 		void SetPrimarySortMode(FName SortedColumnId, EColumnSortMode::Type SortMode)
 		{
-			const TReplicationColumn<TItemType>* Column = FindColumnByName(SortedColumnId);
+			const TSharedPtr<IReplicationTreeColumn<TItemType>> Column = FindColumnByName(SortedColumnId);
 			if (ensureAlways(Column && Column->CanBeSorted()))
 			{
 				PrimarySortInfo = { SortedColumnId, SortMode };
@@ -244,7 +244,7 @@ namespace UE::ConcertSharedSlate
 				return;
 			}
 			
-			const TReplicationColumn<TItemType>* Column = FindColumnByName(SortedColumnId);
+			const TSharedPtr<IReplicationTreeColumn<TItemType>> Column = FindColumnByName(SortedColumnId);
 			if (ensure(Column && Column->CanBeSorted()))
 			{
 				SecondarySortInfo = { SortedColumnId, SortMode };
@@ -290,6 +290,15 @@ namespace UE::ConcertSharedSlate
 		TSharedPtr<STreeView<TSharedPtr<TItemType>>> TreeView;
 		/** The name of the column which will have the SExpandableArrow widget for the tree view. */
 		FName ExpandableColumnId;
+		
+		/** Binds column names to their infos. */
+		TMap<FName, TReplicationColumnEntry<TItemType>> ColumnInfos;
+		/**
+		 * Columns currently being displayed.
+		 * Currently columns are created statically on widget construction but in the future we could add dynamic column registration.
+		 * Hence for now every entry in ColumnInfos will have a corresponding entry here.
+		 */
+		TMap<FName, TSharedRef<IReplicationTreeColumn<TItemType>>> ColumnInstances;
 
 		TArray<TSharedPtr<TItemType>>* AllRootItems = nullptr;
 		/** Contains only the root items that passed the filters */
@@ -347,7 +356,7 @@ namespace UE::ConcertSharedSlate
 		// Sorting
 		void Resort();
 		void Sort(TArray<TSharedPtr<TItemType>>& Items);
-		const TReplicationColumn<TItemType>* FindColumnByName(const FName& ColumnId) const;
+		TSharedPtr<IReplicationTreeColumn<TItemType>> FindColumnByName(const FName& ColumnId) const;
 	};
 
 	template <typename TItemType>
@@ -429,23 +438,28 @@ namespace UE::ConcertSharedSlate
 	template <typename TItemType>
 	TSharedRef<SHeaderRow> SReplicationTreeView<TItemType>::CreateHeaderRow(const FArguments& InArgs)
 	{
-		TArray<TReplicationColumn<TItemType>> Columns = InArgs._Columns;
-		Columns.Sort([](const TReplicationColumn<TItemType>& Left, const TReplicationColumn<TItemType>& Right) { return Left.GetColumnSortOrderValue() < Right.GetColumnSortOrderValue(); });
+		for (const TReplicationColumnEntry<TItemType>& ColumnEntry : InArgs._Columns)
+		{
+			const FName& ColumnId = ColumnEntry.ColumnId;
+			checkf(!ColumnInfos.Contains(ColumnId), TEXT("Duplicate column ID %s"), *ColumnId.ToString());
+			ColumnInfos.Add(ColumnId, ColumnEntry);
+		}
+		
+		TArray<FName> ColumnNames;
+		ColumnInfos.GenerateKeyArray(ColumnNames);
+		ColumnNames.Sort([this](const FName& Left, const FName& Right) { return ColumnInfos[Left].ColumnInfo.SortOrder < ColumnInfos[Right].ColumnInfo.SortOrder; });
 		
 		HeaderRow = SNew(SHeaderRow).Visibility(InArgs._HeaderRowVisibility);
-		TSet<FName> DuplicateColumnDetection;
-		for (TReplicationColumn<TItemType>& Column : Columns)
+		for (const FName& ColumnName : ColumnNames)
 		{
-			const FName ColumnId = Column.ColumnId;
-			check(!DuplicateColumnDetection.Contains(ColumnId));
-			DuplicateColumnDetection.Add(ColumnId);
+			check(ColumnInfos[ColumnName].CreateColumn.IsBound());
+			TSharedRef<IReplicationTreeColumn<TItemType>> ColumnInstance = ColumnInfos[ColumnName].CreateColumn.Execute();
+			
+			const SHeaderRow::FColumn::FArguments HeaderRowArgs = ColumnInstance->CreateHeaderRowArgs();
+			checkf(HeaderRowArgs._ColumnId == ColumnName, TEXT("CreateHeaderRowArgs returned %s but in the TMap was bound to %s"), *HeaderRowArgs._ColumnId.ToString(), *ColumnName.ToString());
+			HeaderRow->AddColumn(HeaderRowArgs);
 			
-			// SHeaderRow owns the columns and deletes them when destroyed
-			TReplicationColumn<TItemType>* ManagedByHeaderRow = new TReplicationColumn<TItemType>(MoveTemp(Column));
-			ManagedByHeaderRow->SortPriority.Bind(TAttribute<EColumnSortPriority::Type>::FGetter::CreateSP(this, &SReplicationTreeView::GetColumnSortPriority, ColumnId));
-			ManagedByHeaderRow->SortMode.Bind(TAttribute<EColumnSortMode::Type>::FGetter::CreateSP(this, &SReplicationTreeView::GetColumnSortMode, ColumnId));
-			ManagedByHeaderRow->OnSortModeChanged.BindSP(this, &SReplicationTreeView::OnColumnSortModeChanged);
-			HeaderRow->AddColumn(*ManagedByHeaderRow);
+			ColumnInstances.Add(ColumnName, MoveTemp(ColumnInstance));
 		}
 
 		return HeaderRow.ToSharedRef();
@@ -527,7 +541,7 @@ namespace UE::ConcertSharedSlate
 		const EColumnSortMode::Type InSortMode
 		)
 	{
-		const TReplicationColumn<TItemType>* Column = FindColumnByName(ColumnId);
+		const TSharedPtr<IReplicationTreeColumn<TItemType>> Column = FindColumnByName(ColumnId);
 		if (!ensure(Column)
 			// Cannot bind
 			|| !Column->CanBeSorted())
@@ -571,11 +585,10 @@ namespace UE::ConcertSharedSlate
 		{
 			return;
 		}
-	
-		for (const SHeaderRow::FColumn& Column : HeaderRow->GetColumns())
+
+		for (const TPair<FName, TSharedRef<IReplicationTreeColumn<TItemType>>>& ColumnEntry : ColumnInstances)
 		{
-			const TReplicationColumn<TItemType>& CastColumn = static_cast<const TReplicationColumn<TItemType>&>(Column);
-			CastColumn.ExecutePopulateSearchString(*Item.Get(), OutSearchStrings);
+			ColumnEntry.Value->PopulateSearchString(*Item, OutSearchStrings);
 		}
 	}
 
@@ -652,7 +665,7 @@ namespace UE::ConcertSharedSlate
 	{
 		const auto IsLessThan = [this](const TSharedPtr<TItemType>& Left, const TSharedPtr<TItemType>& Right, const FName& ColumnName, EColumnSortMode::Type SortMode)
 		{
-			const TReplicationColumn<TItemType>* Column = FindColumnByName(ColumnName);
+			const TSharedPtr<IReplicationTreeColumn<TItemType>> Column = FindColumnByName(ColumnName);
 			if (!ensure(Column) || !ensureMsgf(Column->CanBeSorted(), TEXT("Validate why invariant was broken.")))
 			{
 				return false;
@@ -689,16 +702,10 @@ namespace UE::ConcertSharedSlate
 	}
 
 	template <typename TItemType>
-	const TReplicationColumn<TItemType>* SReplicationTreeView<TItemType>::FindColumnByName(const FName& ColumnId) const
+	TSharedPtr<IReplicationTreeColumn<TItemType>> SReplicationTreeView<TItemType>::FindColumnByName(const FName& ColumnId) const
 	{
-		for (const SHeaderRow::FColumn& Column : HeaderRow->GetColumns())
-		{
-			if (Column.ColumnId == ColumnId)
-			{
-				return static_cast<const TReplicationColumn<TItemType>*>(&Column);
-			}
-		}
-		return nullptr;
+		const TSharedRef<IReplicationTreeColumn<TItemType>>* Instance = ColumnInstances.Find(ColumnId);
+		return Instance ? Instance->ToSharedPtr() : nullptr;
 	}
 }
 
diff --git a/Engine/Plugins/Developer/Concert/ConcertUI/ConcertSharedSlate/Source/ConcertSharedSlate/Private/Replication/ReplicationWidgetFactories.cpp b/Engine/Plugins/Developer/Concert/ConcertUI/ConcertSharedSlate/Source/ConcertSharedSlate/Private/Replication/ReplicationWidgetFactories.cpp
--- a/Engine/Plugins/Developer/Concert/ConcertUI/ConcertSharedSlate/Source/ConcertSharedSlate/Private/Replication/ReplicationWidgetFactories.cpp
+++ b/Engine/Plugins/Developer/Concert/ConcertUI/ConcertSharedSlate/Source/ConcertSharedSlate/Private/Replication/ReplicationWidgetFactories.cpp
@@ -21,7 +21,7 @@ namespace UE::ConcertSharedSlate
 	{
 		return SNew(SBaseReplicationStreamEditor, MoveTemp(EditorParams.DataModel), MoveTemp(EditorParams.ObjectSource), MoveTemp(EditorParams.PropertySource))
 			.PropertyTreeView(MoveTemp(ViewerParams.PropertyTreeView))
-			.AdditionalObjectColumns(MoveTemp(ViewerParams.AdditionalObjectColumns))
+			.ObjectColumns(MoveTemp(ViewerParams.ObjectColumns))
 			.PrimaryObjectSort(ViewerParams.PrimaryObjectSort)
 			.SecondaryObjectSort(ViewerParams.SecondaryObjectSort)
 			.ObjectHierarchy(MoveTemp(ViewerParams.ObjectHierarchy))
@@ -35,10 +35,10 @@ namespace UE::ConcertSharedSlate
 	
 	TSharedRef<IPropertyTreeView> CreateSearchablePropertyTreeView(FCreatePropertyTreeViewParams Params)
 	{
-		// The label columns is always required
-		const bool bHasLabel = Params.PropertyColumns.ContainsByPredicate([](const ReplicationColumns::FReplicationPropertyColumn& Column)
+		// The label column is always required
+		const bool bHasLabel = Params.PropertyColumns.ContainsByPredicate([](const FPropertyColumnEntry& Entry)
 		{
-			return Column.ColumnId == ReplicationColumns::Property::LabelColumnId;
+			return Entry.ColumnId == ReplicationColumns::Property::LabelColumnId;
 		});
 		if (!bHasLabel)
 		{
diff --git a/Engine/Plugins/Developer/Concert/ConcertUI/ConcertSharedSlate/Source/ConcertSharedSlate/Public/Replication/Editor/View/Column/IObjectTreeColumn.h b/Engine/Plugins/Developer/Concert/ConcertUI/ConcertSharedSlate/Source/ConcertSharedSlate/Public/Replication/Editor/View/Column/IObjectTreeColumn.h
new file mode 100644
--- /dev/null
+++ b/Engine/Plugins/Developer/Concert/ConcertUI/ConcertSharedSlate/Source/ConcertSharedSlate/Public/Replication/Editor/View/Column/IObjectTreeColumn.h
@@ -0,0 +1,20 @@
+// Copyright Epic Games, Inc. All Rights Reserved.
+
+#pragma once
+
+#include "IReplicationTreeColumn.h"
+#include "ReplicationColumnInfo.h"
+#include "Replication/Editor/Model/ReplicatedObjectData.h"
+
+namespace UE::ConcertSharedSlate
+{
+	/** Holds information relevant to a row. */
+	struct FObjectTreeRowContext
+	{
+		/** The data that is stored in the row */
+		FReplicatedObjectData RowData;
+	};
+	
+	using IObjectTreeColumn = IReplicationTreeColumn<FObjectTreeRowContext>;
+	using FObjectColumnEntry = TReplicationColumnEntry<FObjectTreeRowContext>;
+}
diff --git a/Engine/Plugins/Developer/Concert/ConcertUI/ConcertSharedSlate/Source/ConcertSharedSlate/Public/Replication/Editor/View/Column/IPropertyTreeColumn.h b/Engine/Plugins/Developer/Concert/ConcertUI/ConcertSharedSlate/Source/ConcertSharedSlate/Public/Replication/Editor/View/Column/IPropertyTreeColumn.h
new file mode 100644
--- /dev/null
+++ b/Engine/Plugins/Developer/Concert/ConcertUI/ConcertSharedSlate/Source/ConcertSharedSlate/Public/Replication/Editor/View/Column/IPropertyTreeColumn.h
@@ -0,0 +1,20 @@
+// Copyright Epic Games, Inc. All Rights Reserved.
+
+#pragma once
+
+#include "IReplicationTreeColumn.h"
+#include "ReplicationColumnInfo.h"
+#include "Replication/Editor/Model/ReplicatedPropertyData.h"
+
+namespace UE::ConcertSharedSlate
+{
+	/** Holds information relevant to a row. */
+	struct FPropertyTreeRowContext
+	{
+		/** The data that is stored in the row */
+		FReplicatedPropertyData RowData;
+	};
+	
+	using IPropertyTreeColumn = IReplicationTreeColumn<FPropertyTreeRowContext>;
+	using FPropertyColumnEntry = TReplicationColumnEntry<FPropertyTreeRowContext>;
+}
diff --git a/Engine/Plugins/Developer/Concert/ConcertUI/ConcertSharedSlate/Source/ConcertSharedSlate/Public/Replication/Editor/View/Column/IReplicationTreeColumn.h b/Engine/Plugins/Developer/Concert/ConcertUI/ConcertSharedSlate/Source/ConcertSharedSlate/Public/Replication/Editor/View/Column/IReplicationTreeColumn.h
new file mode 100644
--- /dev/null
+++ b/Engine/Plugins/Developer/Concert/ConcertUI/ConcertSharedSlate/Source/ConcertSharedSlate/Public/Replication/Editor/View/Column/IReplicationTreeColumn.h
@@ -0,0 +1,35 @@
+// Copyright Epic Games, Inc. All Rights Reserved.
+
+#pragma once
+
+#include "Containers/ContainersFwd.h"
+#include "Templates/SharedPointer.h"
+#include "Widgets/Views/SHeaderRow.h"
+
+class FString;
+class FText;
+
+namespace UE::ConcertSharedSlate
+{
+	template<typename TTreeItemType>
+	class IReplicationTreeColumn
+	{
+	public:
+
+		struct FBuildArgs
+		{
+			TSharedPtr<FText> HighlightText;
+			TTreeItemType RowItem;
+		};
+		
+		virtual SHeaderRow::FColumn::FArguments CreateHeaderRowArgs() const = 0;
+
+		virtual TSharedRef<SWidget> GenerateColumnWidget(const FBuildArgs& InArgs) = 0;
+		virtual void PopulateSearchString(const TTreeItemType& InItem, TArray<FString>& InOutSearchStrings) const {}
+
+		virtual bool CanBeSorted() const { return false; }
+		virtual bool IsLessThan(const TTreeItemType& Left, const TTreeItemType& Right) const { return false; }
+
+		virtual ~IReplicationTreeColumn() = default;
+	};
+}
\ No newline at end of file
diff --git a/Engine/Plugins/Developer/Concert/ConcertUI/ConcertSharedSlate/Source/ConcertSharedSlate/Public/Replication/Editor/View/Column/ReplicationColumnDelegates.h b/Engine/Plugins/Developer/Concert/ConcertUI/ConcertSharedSlate/Source/ConcertSharedSlate/Public/Replication/Editor/View/Column/ReplicationColumnDelegates.h
new file mode 100644
--- /dev/null
+++ b/Engine/Plugins/Developer/Concert/ConcertUI/ConcertSharedSlate/Source/ConcertSharedSlate/Public/Replication/Editor/View/Column/ReplicationColumnDelegates.h
@@ -0,0 +1,14 @@
+// Copyright Epic Games, Inc. All Rights Reserved.
+
+#pragma once
+
+#include "IReplicationTreeColumn.h"
+
+namespace UE::ConcertSharedSlate
+{
+	template<typename TTreeItemType>
+	struct TReplicationColumnDelegates
+	{
+		DECLARE_DELEGATE_RetVal(TSharedRef<IReplicationTreeColumn<TTreeItemType>>, FCreateColumn);
+	};
+}
diff --git a/Engine/Plugins/Developer/Concert/ConcertUI/ConcertSharedSlate/Source/ConcertSharedSlate/Public/Replication/Editor/View/Column/ReplicationColumnInfo.h b/Engine/Plugins/Developer/Concert/ConcertUI/ConcertSharedSlate/Source/ConcertSharedSlate/Public/Replication/Editor/View/Column/ReplicationColumnInfo.h
new file mode 100644
--- /dev/null
+++ b/Engine/Plugins/Developer/Concert/ConcertUI/ConcertSharedSlate/Source/ConcertSharedSlate/Public/Replication/Editor/View/Column/ReplicationColumnInfo.h
@@ -0,0 +1,31 @@
+// Copyright Epic Games, Inc. All Rights Reserved.
+
+#pragma once
+
+#include "ReplicationColumnDelegates.h"
+
+#include "HAL/Platform.h"
+
+namespace UE::ConcertSharedSlate
+{
+	/**
+	 * Meta information about a column that can change independent of the IReplicationTreeColumn implementation.
+	 * For example, IReplicationTreeColumn does not need to know about sort order.
+	 */
+	struct FReplicationColumnInfo
+	{
+		/** Determines whether this column is the first, etc. */
+		int32 SortOrder = 0;
+	};
+
+	template<typename TTreeItemType>
+	struct TReplicationColumnEntry
+	{
+		/** Factory function for creating the column instance. */
+		typename TReplicationColumnDelegates<TTreeItemType>::FCreateColumn CreateColumn;
+		/** The ID of the created column. Must be equal to what IReplicationTreeColumn::CreateHeaderRowArgs would return. */
+		FName ColumnId;
+		/** Additional info about the column that is not directly controlled by IReplicationTreeColumn. */
+		FReplicationColumnInfo ColumnInfo;
+	};
+}
diff --git a/Engine/Plugins/Developer/Concert/ConcertUI/ConcertSharedSlate/Source/ConcertSharedSlate/Public/Replication/Editor/View/Column/ReplicationColumnsUtils.h b/Engine/Plugins/Developer/Concert/ConcertUI/ConcertSharedSlate/Source/ConcertSharedSlate/Public/Replication/Editor/View/Column/ReplicationColumnsUtils.h
new file mode 100644
--- /dev/null
+++ b/Engine/Plugins/Developer/Concert/ConcertUI/ConcertSharedSlate/Source/ConcertSharedSlate/Public/Replication/Editor/View/Column/ReplicationColumnsUtils.h
@@ -0,0 +1,156 @@
+// Copyright Epic Games, Inc. All Rights Reserved.
+
+#pragma once
+
+#include "Replication/Editor/View/Column/IReplicationTreeColumn.h"
+#include "Replication/Editor/View/Column/ReplicationColumnDelegates.h"
+#include "Replication/Editor/View/Column/ReplicationColumnInfo.h"
+
+#include "Containers/ArrayView.h"
+#include "Widgets/Input/SCheckBox.h"
+#include "Templates/SharedPointer.h"
+
+namespace UE::ConcertSharedSlate
+{
+	class IEditableReplicationStreamModel;
+	class IReplicationStreamViewer;
+	class IReplicationStreamModel;
+}
+
+namespace UE::ConcertSharedSlate
+{
+	/** Info about a column that is being sorted. */
+	struct FColumnSortInfo
+	{
+		FName SortedColumnId;
+		EColumnSortMode::Type SortMode = EColumnSortMode::None;
+
+		bool IsValid() const { return SortedColumnId != NAME_None && SortMode != EColumnSortMode::None; }
+	};
+	
+	template<typename TListItemType>
+	struct TCheckboxColumnDelegates
+	{
+		DECLARE_DELEGATE_RetVal_OneParam(ECheckBoxState, FGetColumnCheckboxState, const TListItemType& /*ObjectData*/);
+		DECLARE_DELEGATE_TwoParams(FOnColumnCheckboxChanged, bool /*bIsChecked*/, const TListItemType& /*ObjectData*/);
+		DECLARE_DELEGATE_RetVal_OneParam(bool, FIsEnabled, const TListItemType& /*ObjectData*/);
+		DECLARE_DELEGATE_RetVal_OneParam(FText, FGetToolTipText, const TListItemType& /*ObjectData*/);
+
+		TCheckboxColumnDelegates(
+			FGetColumnCheckboxState InGetCheckboxStateDelegate,
+			FOnColumnCheckboxChanged InOnCheckboxChangedDelegate,
+			FGetToolTipText InGetToolTipTextDelegate = {},
+			FIsEnabled InIsEnabledDelegate = {}
+			)
+			: GetCheckboxStateDelegate(MoveTemp(InGetCheckboxStateDelegate))
+			, OnCheckboxChangedDelegate(MoveTemp(InOnCheckboxChangedDelegate))
+			, GetToolTipTextDelegate(MoveTemp(InGetToolTipTextDelegate))
+			, IsEnabledDelegate(MoveTemp(InIsEnabledDelegate))
+		{}
+		
+		TCheckboxColumnDelegates(
+			FGetColumnCheckboxState InGetCheckboxStateDelegate,
+			FOnColumnCheckboxChanged InOnCheckboxChangedDelegate,
+			TAttribute<FText> InToolTipTextAttribute,
+			FIsEnabled InIsEnabledDelegate = {}
+			)
+			: GetCheckboxStateDelegate(MoveTemp(InGetCheckboxStateDelegate))
+			, OnCheckboxChangedDelegate(MoveTemp(InOnCheckboxChangedDelegate))
+			, GetToolTipTextDelegate(FGetToolTipText::CreateLambda([GetToolTipText = MoveTemp(InToolTipTextAttribute)](const TListItemType&){ return GetToolTipText.Get(); }))
+			, IsEnabledDelegate(MoveTemp(InIsEnabledDelegate))
+		{}
+
+		FGetColumnCheckboxState GetCheckboxStateDelegate;
+		FOnColumnCheckboxChanged OnCheckboxChangedDelegate;
+
+		/** Optional. Defaults to FText::GetEmpty() */
+		FGetToolTipText GetToolTipTextDelegate;
+		/** Optional. Default to true. */
+		FIsEnabled IsEnabledDelegate;
+	};
+
+	/** Util for making a checkbox column */
+	template<typename TListItemType>
+	UE::ConcertSharedSlate::template TReplicationColumnEntry<TListItemType> MakeCheckboxColumn(
+		FName ColumnId,
+		TCheckboxColumnDelegates<TListItemType> Delegates,
+		FText DefaultLabel,
+		const int32 Priority,
+		const float ColumnWidth = 20.f
+		)
+	{
+		class TCheckboxColumn : public UE::ConcertSharedSlate::template IReplicationTreeColumn<TListItemType>
+		{
+		public:
+
+			TCheckboxColumn(
+				FName ColumnId,
+				FText DefaultLabel,
+				const float ColumnWidth,
+				TCheckboxColumnDelegates<TListItemType> Delegates
+				)
+				: ColumnId(ColumnId)
+				, DefaultLabel(MoveTemp(DefaultLabel))
+				, ColumnWidth(ColumnWidth)
+				, Delegates(MoveTemp(Delegates))
+			{}
+			
+			virtual SHeaderRow::FColumn::FArguments CreateHeaderRowArgs() const override
+			{
+				return SHeaderRow::Column(ColumnId)
+					.DefaultLabel(DefaultLabel)
+					.FixedWidth(ColumnWidth);
+			}
+			virtual TSharedRef<SWidget> GenerateColumnWidget(const typename IReplicationTreeColumn<TListItemType>::FBuildArgs& InArgs) override
+			{
+				return SNew(SCheckBox)
+					.ToolTipText_Lambda([GetToolTipDelegate = Delegates.GetToolTipTextDelegate, RowData = InArgs.RowItem]()
+					{
+						return GetToolTipDelegate.IsBound() ? GetToolTipDelegate.Execute(RowData) : FText::GetEmpty();
+					})
+					.IsEnabled_Lambda([IsEnabledDelegate = Delegates.IsEnabledDelegate, RowData = InArgs.RowItem]()
+					{
+						return !IsEnabledDelegate.IsBound() || IsEnabledDelegate.Execute(RowData);
+					})
+					.IsChecked_Lambda([IsCheckedDelegate = Delegates.GetCheckboxStateDelegate, RowData = InArgs.RowItem]()
+					{
+						return IsCheckedDelegate.Execute(RowData);
+					})
+					.OnCheckStateChanged_Lambda([OnChangedDelegate = Delegates.OnCheckboxChangedDelegate, RowData = InArgs.RowItem](ECheckBoxState NewState)
+					{
+						const bool bIsChecked = NewState == ECheckBoxState::Checked;
+						OnChangedDelegate.Execute(bIsChecked, RowData);
+					});
+			}
+			
+			virtual void PopulateSearchString(const TListItemType& InItem, TArray<FString>& InOutSearchStrings) const override {}
+			
+			virtual bool CanBeSorted() const override { return true; }
+			virtual bool IsLessThan(const TListItemType& Left, const TListItemType& Right) const override
+			{
+				const bool bIsLeftChecked = Delegates.GetCheckboxStateDelegate.IsBound() && Delegates.GetCheckboxStateDelegate.Execute(Left) == ECheckBoxState::Checked;
+				const bool bIsRightChecked = Delegates.GetCheckboxStateDelegate.IsBound() && Delegates.GetCheckboxStateDelegate.Execute(Right) == ECheckBoxState::Checked;
+				// Less for checkbox means: checked < unchecked. !bIsRightChecked because checked == checked so it cannot be less.
+				return bIsLeftChecked && !bIsRightChecked;
+			}
+
+		private:
+
+			const FName ColumnId;
+			const FText DefaultLabel;
+			const float ColumnWidth;
+			const TCheckboxColumnDelegates<TListItemType> Delegates;
+		};
+		
+		check(Delegates.GetCheckboxStateDelegate.IsBound() && Delegates.OnCheckboxChangedDelegate.IsBound());
+		return 
+		{
+			UE::ConcertSharedSlate::template TReplicationColumnDelegates<TListItemType>::FCreateColumn::CreateLambda([ColumnId, DefaultLabel, ColumnWidth, Delegates = MoveTemp(Delegates)]()
+			{
+				return MakeShared<TCheckboxColumn>(ColumnId, DefaultLabel, ColumnWidth, Delegates);
+			}),
+			ColumnId,
+			{ Priority }
+		};
+	}
+}
diff --git a/Engine/Plugins/Developer/Concert/ConcertUI/ConcertSharedSlate/Source/ConcertSharedSlate/Public/Replication/Editor/View/Column/SelectionViewerColumns.h b/Engine/Plugins/Developer/Concert/ConcertUI/ConcertSharedSlate/Source/ConcertSharedSlate/Public/Replication/Editor/View/Column/SelectionViewerColumns.h
new file mode 100644
--- /dev/null
+++ b/Engine/Plugins/Developer/Concert/ConcertUI/ConcertSharedSlate/Source/ConcertSharedSlate/Public/Replication/Editor/View/Column/SelectionViewerColumns.h
@@ -0,0 +1,54 @@
+// Copyright Epic Games, Inc. All Rights Reserved.
+
+#pragma once
+
+#include "IObjectTreeColumn.h"
+#include "IPropertyTreeColumn.h"
+#include "ReplicationColumnInfo.h"
+
+#include "Templates/SharedPointer.h"
+
+namespace UE::ConcertSharedSlate
+{
+	struct FObjectTreeRowContext;
+	struct FPropertyTreeRowContext;
+	class IEditableReplicationStreamModel;
+	class IReplicationStreamViewer;
+	class IReplicationStreamModel;
+	class IObjectNameModel;
+}
+
+namespace UE::ConcertSharedSlate::ReplicationColumns::TopLevel
+{
+	CONCERTSHAREDSLATE_API extern const FName IconColumnId;
+	CONCERTSHAREDSLATE_API extern const FName LabelColumnId;
+	CONCERTSHAREDSLATE_API extern const FName TypeColumnId;
+
+	enum class ETopLevelColumnOrder : int32
+	{
+		/** Label of the object */
+		Label = 20,
+		/** Class of the object */
+		Type = 30,
+	};
+
+	CONCERTSHAREDSLATE_API FObjectColumnEntry LabelColumn(TSharedRef<IReplicationStreamModel> Model, IObjectNameModel* OptionalNameModel = nullptr);
+	CONCERTSHAREDSLATE_API FObjectColumnEntry TypeColumn(TSharedRef<IReplicationStreamModel> Model);
+}
+
+namespace UE::ConcertSharedSlate::ReplicationColumns::Property
+{
+	CONCERTSHAREDSLATE_API extern const FName LabelColumnId;
+	CONCERTSHAREDSLATE_API extern const FName TypeColumnId;
+	
+	enum class EReplicationPropertyColumnOrder : int32
+	{
+		/** Label of the property */
+		Label = 10,
+		/** Type of the property */
+		Type = 20
+	};
+	
+	CONCERTSHAREDSLATE_API FPropertyColumnEntry LabelColumn();
+	CONCERTSHAREDSLATE_API FPropertyColumnEntry TypeColumn();
+}
diff --git a/Engine/Plugins/Developer/Concert/ConcertUI/ConcertSharedSlate/Source/ConcertSharedSlate/Public/Replication/Editor/View/PredefinedReplicationColumns.h b/Engine/Plugins/Developer/Concert/ConcertUI/ConcertSharedSlate/Source/ConcertSharedSlate/Public/Replication/Editor/View/PredefinedReplicationColumns.h
deleted file mode 100644
diff --git a/Engine/Plugins/Developer/Concert/ConcertUI/ConcertSharedSlate/Source/ConcertSharedSlate/Public/Replication/Editor/View/ReplicationColumn.h b/Engine/Plugins/Developer/Concert/ConcertUI/ConcertSharedSlate/Source/ConcertSharedSlate/Public/Replication/Editor/View/ReplicationColumn.h
deleted file mode 100644
diff --git a/Engine/Plugins/Developer/Concert/ConcertUI/ConcertSharedSlate/Source/ConcertSharedSlate/Public/Replication/Editor/View/ReplicationColumnsUtils.h b/Engine/Plugins/Developer/Concert/ConcertUI/ConcertSharedSlate/Source/ConcertSharedSlate/Public/Replication/Editor/View/ReplicationColumnsUtils.h
deleted file mode 100644
diff --git a/Engine/Plugins/Developer/Concert/ConcertUI/ConcertSharedSlate/Source/ConcertSharedSlate/Public/Replication/Editor/View/SelectionViewerColumns.h b/Engine/Plugins/Developer/Concert/ConcertUI/ConcertSharedSlate/Source/ConcertSharedSlate/Public/Replication/Editor/View/SelectionViewerColumns.h
deleted file mode 100644
diff --git a/Engine/Plugins/Developer/Concert/ConcertUI/ConcertSharedSlate/Source/ConcertSharedSlate/Public/Replication/ReplicationWidgetFactories.h b/Engine/Plugins/Developer/Concert/ConcertUI/ConcertSharedSlate/Source/ConcertSharedSlate/Public/Replication/ReplicationWidgetFactories.h
--- a/Engine/Plugins/Developer/Concert/ConcertUI/ConcertSharedSlate/Source/ConcertSharedSlate/Public/Replication/ReplicationWidgetFactories.h
+++ b/Engine/Plugins/Developer/Concert/ConcertUI/ConcertSharedSlate/Source/ConcertSharedSlate/Public/Replication/ReplicationWidgetFactories.h
@@ -2,12 +2,14 @@
 
 #pragma once
 
-#include "Editor/View/PredefinedReplicationColumns.h"
+#include "Editor/View/Column/IObjectTreeColumn.h"
+#include "Editor/View/Column/IPropertyTreeColumn.h"
+#include "Editor/View/Column/ReplicationColumnInfo.h"
+#include "Editor/View/Column/SelectionViewerColumns.h"
+#include "Replication/Editor/View/Column/ReplicationColumnsUtils.h"
 #include "ReplicationWidgetDelegates.h"
 
 #include "Delegates/Delegate.h"
-#include "Editor/View/ReplicationColumnsUtils.h"
-#include "Editor/View/SelectionViewerColumns.h"
 #include "Misc/Attribute.h"
 #include "Templates/SharedPointer.h"
 
@@ -59,8 +61,8 @@ namespace UE::ConcertSharedSlate
 	
 	struct FCreatePropertyTreeViewParams
 	{
-		/** Optional, additional columns to add to the property view. The LabelColumn will always be added. */
-		TArray<ReplicationColumns::FReplicationPropertyColumn> PropertyColumns
+		/** Optional. Additional property columns you want added. */
+		TArray<FPropertyColumnEntry> PropertyColumns
 		{
 			ReplicationColumns::Property::LabelColumn(),
 			ReplicationColumns::Property::TypeColumn()
@@ -109,8 +111,9 @@ namespace UE::ConcertSharedSlate
 		/** Called to generate the context menu for objects. This extends the options already generated by this widget and is called at the end. */
 		FExtendObjectMenu OnExtendObjectsContextMenu;
 		
-		/** Additional columns to add to the object view */
-		TArray<ReplicationColumns::FReplicationTopLevelObjectColumn> AdditionalObjectColumns;
+		/** Optional. Additional object columns you want added. */
+		TArray<FObjectColumnEntry> ObjectColumns;
+		
 		/** Optional initial primary sort mode for object rows */
 		FColumnSortInfo PrimaryObjectSort;
 		/** Optional initial secondary sort mode for object rows */
