MU Replication: When switching levels, the user's stream and authority should be cleared

diff --git a/Engine/Plugins/Developer/Concert/ConcertApp/MultiUserClient/Source/MultiUserClient/Private/Replication/MultiUserReplicationManager.cpp b/Engine/Plugins/Developer/Concert/ConcertApp/MultiUserClient/Source/MultiUserClient/Private/Replication/MultiUserReplicationManager.cpp
--- a/Engine/Plugins/Developer/Concert/ConcertApp/MultiUserClient/Source/MultiUserClient/Private/Replication/MultiUserReplicationManager.cpp
+++ b/Engine/Plugins/Developer/Concert/ConcertApp/MultiUserClient/Source/MultiUserClient/Private/Replication/MultiUserReplicationManager.cpp
@@ -197,5 +197,6 @@ namespace UE::MultiUserClient
 
 	FMultiUserReplicationManager::FConnectedState::FConnectedState(TSharedRef<IConcertSyncClient> InClient, FReplicationDiscoveryContainer& InDiscoveryContainer)
 		: ClientManager(InClient, InClient->GetConcertClient()->GetCurrentSession().ToSharedRef(), InDiscoveryContainer)
+		, ChangeLevelHandler(ClientManager.GetLocalClient().GetClientEditModel().Get())
 	{}
 }
diff --git a/Engine/Plugins/Developer/Concert/ConcertApp/MultiUserClient/Source/MultiUserClient/Private/Replication/MultiUserReplicationManager.h b/Engine/Plugins/Developer/Concert/ConcertApp/MultiUserClient/Source/MultiUserClient/Private/Replication/MultiUserReplicationManager.h
--- a/Engine/Plugins/Developer/Concert/ConcertApp/MultiUserClient/Source/MultiUserClient/Private/Replication/MultiUserReplicationManager.h
+++ b/Engine/Plugins/Developer/Concert/ConcertApp/MultiUserClient/Source/MultiUserClient/Private/Replication/MultiUserReplicationManager.h
@@ -6,6 +6,7 @@
 #include "IConcertSession.h"
 #include "Replication/IMultiUserReplication.h"
 #include "Replication/Stream/Discovery/ReplicationDiscoveryContainer.h"
+#include "UnrealEditor/ChangeLevelHandler.h"
 
 #include "Misc/Optional.h"
 #include "Templates/SharedPointer.h"
@@ -101,6 +102,9 @@ namespace UE::MultiUserClient
 			 * Only valid when ConnectionState == EMultiUserReplicationConnectionState::Connected.
 			 */
 			FReplicationClientManager ClientManager;
+
+			/** Clears local client's registered objects when leaving map. */
+			FChangeLevelHandler ChangeLevelHandler;
 			
 			FConnectedState(TSharedRef<IConcertSyncClient> InClient, FReplicationDiscoveryContainer& InDiscoveryContainer);
 		};
diff --git a/Engine/Plugins/Developer/Concert/ConcertApp/MultiUserClient/Source/MultiUserClient/Private/Replication/UnrealEditor/ChangeLevelHandler.cpp b/Engine/Plugins/Developer/Concert/ConcertApp/MultiUserClient/Source/MultiUserClient/Private/Replication/UnrealEditor/ChangeLevelHandler.cpp
new file mode 100644
--- /dev/null
+++ b/Engine/Plugins/Developer/Concert/ConcertApp/MultiUserClient/Source/MultiUserClient/Private/Replication/UnrealEditor/ChangeLevelHandler.cpp
@@ -0,0 +1,48 @@
+// Copyright Epic Games, Inc. All Rights Reserved.
+
+#include "ChangeLevelHandler.h"
+
+#include "Replication/Editor/Model/IEditableReplicationStreamModel.h"
+
+#include "Engine/Engine.h"
+#include "Engine/World.h"
+
+namespace UE::MultiUserClient
+{
+	FChangeLevelHandler::FChangeLevelHandler(ConcertSharedSlate::IEditableReplicationStreamModel& UpdatedModel)
+		: UpdatedModel(UpdatedModel)
+	{
+		if (ensure(GEngine))
+		{
+			GEngine->OnWorldDestroyed().AddRaw(this, &FChangeLevelHandler::OnWorldDestroyed);
+			GEngine->OnWorldAdded().AddRaw(this, &FChangeLevelHandler::OnWorldAdded);
+		}
+	}
+
+	FChangeLevelHandler::~FChangeLevelHandler()
+	{
+		if (GEngine)
+		{
+			GEngine->OnWorldDestroyed().RemoveAll(this);
+			GEngine->OnWorldAdded().RemoveAll(this);
+		}
+	}
+
+	void FChangeLevelHandler::OnWorldDestroyed(UWorld* World)
+	{
+		// Remember the current map so when the next world
+		if (World->WorldType == EWorldType::Editor)
+		{
+			PreviousWorldPath = FSoftObjectPath(World);
+		}
+	}
+
+	void FChangeLevelHandler::OnWorldAdded(UWorld* World) const
+	{
+		// User reloaded the same map - keep around all settings
+		if (PreviousWorldPath != World)
+		{
+			UpdatedModel.Clear();
+		}
+	}
+}
diff --git a/Engine/Plugins/Developer/Concert/ConcertApp/MultiUserClient/Source/MultiUserClient/Private/Replication/UnrealEditor/ChangeLevelHandler.h b/Engine/Plugins/Developer/Concert/ConcertApp/MultiUserClient/Source/MultiUserClient/Private/Replication/UnrealEditor/ChangeLevelHandler.h
new file mode 100644
--- /dev/null
+++ b/Engine/Plugins/Developer/Concert/ConcertApp/MultiUserClient/Source/MultiUserClient/Private/Replication/UnrealEditor/ChangeLevelHandler.h
@@ -0,0 +1,40 @@
+// Copyright Epic Games, Inc. All Rights Reserved.
+
+#pragma once
+
+#include "HAL/Platform.h"
+#include "Templates/UnrealTemplate.h"
+#include "UObject/SoftObjectPath.h"
+
+class UWorld;
+
+namespace UE::ConcertSharedSlate { class IEditableReplicationStreamModel; }
+
+namespace UE::MultiUserClient
+{
+	/**
+	 * Handles a client opening a new level.
+	 * Upon leaving a level, all replicated objects assigned to that client in that level are removed.
+	 */
+	class FChangeLevelHandler : public FNoncopyable
+	{
+	public:
+		
+		/**
+		 * @param UpdatedModel The client model to update when the local editor changes maps. The caller ensures that it outlives the constructed object.
+		 */
+		FChangeLevelHandler(ConcertSharedSlate::IEditableReplicationStreamModel& UpdatedModel UE_LIFETIMEBOUND);
+		~FChangeLevelHandler();
+
+	private:
+
+		/** The client model to update when the local editor changes maps. */
+		ConcertSharedSlate::IEditableReplicationStreamModel& UpdatedModel;
+
+		/** Path to the previously open world */
+		FSoftObjectPath PreviousWorldPath;
+		
+		void OnWorldDestroyed(UWorld* World);
+		void OnWorldAdded(UWorld* World) const;
+	};
+}
diff --git a/Engine/Plugins/Developer/Concert/ConcertUI/ConcertSharedSlate/Source/ConcertSharedSlate/Public/Replication/Editor/Model/IEditableReplicationStreamModel.h b/Engine/Plugins/Developer/Concert/ConcertUI/ConcertSharedSlate/Source/ConcertSharedSlate/Public/Replication/Editor/Model/IEditableReplicationStreamModel.h
--- a/Engine/Plugins/Developer/Concert/ConcertUI/ConcertSharedSlate/Source/ConcertSharedSlate/Public/Replication/Editor/Model/IEditableReplicationStreamModel.h
+++ b/Engine/Plugins/Developer/Concert/ConcertUI/ConcertSharedSlate/Source/ConcertSharedSlate/Public/Replication/Editor/Model/IEditableReplicationStreamModel.h
@@ -45,7 +45,7 @@ namespace UE::ConcertSharedSlate
 		virtual void AddProperties(const FSoftObjectPath& Object, TConstArrayView<FConcertPropertyChain> Properties) = 0;
 		/** Removes these properties from the object's list of selected properties. */
 		virtual void RemoveProperties(const FSoftObjectPath& Object, TConstArrayView<FConcertPropertyChain> Properties) = 0;
-
+		
 		/** Called when the object list changes. AddedObjects and RemovedObjects are empty if and only if ChangeReason == EReplicatedObjectChangeReason::Transacted. */
 		DECLARE_MULTICAST_DELEGATE_ThreeParams(FOnObjectsChanged, TConstArrayView<UObject*> AddedObjects, TConstArrayView<FSoftObjectPath> RemovedObjects, EReplicatedObjectChangeReason ChangeReason);
 		virtual FOnObjectsChanged& OnObjectsChanged() = 0;
@@ -53,5 +53,8 @@ namespace UE::ConcertSharedSlate
 		/** Called when the property list of some object changes. If the properties are removed because the object is removed outright, FOnObjectsChanged is called instead.  */
 		DECLARE_MULTICAST_DELEGATE(FOnPropertiesChanged);
 		virtual FOnPropertiesChanged& OnPropertiesChanged() = 0;
+
+		/** Clear all objects in the model */
+		void Clear() { RemoveObjects(GetReplicatedObjects().Array()); }
 	};
 }
\ No newline at end of file
