Replication Tree View
- no longer resorts properties based on whether they're being replication
- auto expands child properties when searching by text

diff --git a/Engine/Plugins/Developer/Concert/ConcertApp/MultiUserClient/Source/MultiUserClient/Private/Widgets/ActiveSession/Replication/Client/Multi/SMultiClientView.cpp b/Engine/Plugins/Developer/Concert/ConcertApp/MultiUserClient/Source/MultiUserClient/Private/Widgets/ActiveSession/Replication/Client/Multi/SMultiClientView.cpp
--- a/Engine/Plugins/Developer/Concert/ConcertApp/MultiUserClient/Source/MultiUserClient/Private/Widgets/ActiveSession/Replication/Client/Multi/SMultiClientView.cpp
+++ b/Engine/Plugins/Developer/Concert/ConcertApp/MultiUserClient/Source/MultiUserClient/Private/Widgets/ActiveSession/Replication/Client/Multi/SMultiClientView.cpp
@@ -90,8 +90,7 @@ namespace UE::MultiUserClient
 				ReplicationColumns::Property::LabelColumn(),
 				ReplicationColumns::Property::TypeColumn(),
 				MultiStreamColumns::AssignPropertyColumn(MultiStreamEditorAttribute, InConcertClient, InClientManager)
-			},
-			.PrimaryPropertySort = { MultiStreamColumns::AssignPropertyColumnId, EColumnSortMode::Ascending }
+			}
 		};
 		TSharedRef<IPropertyTreeView> PropertyTreeView = CreateFilterablePropertyTreeView(MoveTemp(TreeViewParams));
 		
diff --git a/Engine/Plugins/Developer/Concert/ConcertUI/ConcertSharedSlate/Source/ConcertSharedSlate/Private/Replication/Editor/View/Tree/SReplicationTreeView.h b/Engine/Plugins/Developer/Concert/ConcertUI/ConcertSharedSlate/Source/ConcertSharedSlate/Private/Replication/Editor/View/Tree/SReplicationTreeView.h
--- a/Engine/Plugins/Developer/Concert/ConcertUI/ConcertSharedSlate/Source/ConcertSharedSlate/Private/Replication/Editor/View/Tree/SReplicationTreeView.h
+++ b/Engine/Plugins/Developer/Concert/ConcertUI/ConcertSharedSlate/Source/ConcertSharedSlate/Private/Replication/Editor/View/Tree/SReplicationTreeView.h
@@ -306,7 +306,20 @@ namespace UE::ConcertSharedSlate
 		/** Includes ALL items in the hierarchy that have passed the filter. */
 		TSet<TSharedPtr<TItemType>> AllFilteredItems;
 
+		struct FItemMetaData
+		{
+			/** Whether the item is expanded */
+			bool bIsExpanded = false;
+		};
+		/** Additional data about items needed for misc operations, like expansion. */
+		TMap<TSharedPtr<TItemType>, FItemMetaData> ItemMetaData;
+
+		/** Whether all items should be expanded. True while searching and false otherwise. */
+		bool bForceParentItemsExpanded = false;
+
+		/** Whether to call ReapplyFilters next Tick(). */
 		bool bFilterChanged = false;
+		/** Whether to call Resort next Tick(). */
 		bool bRequestedSort = false;
 
 		FColumnSortInfo PrimarySortInfo;
@@ -347,12 +360,18 @@ namespace UE::ConcertSharedSlate
 			ItemOrChildrenPassFilter,
 			NoneInHierarchyPassFilter
 		};
-		
 		void PopulateSearchStrings(const TSharedPtr<TItemType>& Item, TArray<FString>& OutSearchStrings);
 		void ReapplyFilters();
 		EFilterResult ApplyFiltersRecursive(const TSharedPtr<TItemType>& Item, TSet<TSharedPtr<TItemType>>& FilteredItemsToShow);
 		bool PassesFilters(const TSharedPtr<TItemType>& Item);
 
+		/** Called after RootItems has changed. Removes all invalidated ItemMetaData entries */
+		void CleanseItemMetaData();
+		/** Applies the cached expansion states to all items */
+		void ReapplyExpansionStates();
+		/** Callback into tree view when expansion state is changed. */
+		void OnItemExpansionChanged(TSharedPtr<TItemType> Item, bool bIsExpanded);
+
 		// Sorting
 		void Resort();
 		void Sort(TArray<TSharedPtr<TItemType>>& Items);
@@ -408,6 +427,7 @@ namespace UE::ConcertSharedSlate
 					.OnGenerateRow(this, &SReplicationTreeView::OnGenerateRowWidget)
 					.OnContextMenuOpening(InArgs._OnContextMenuOpening)
 					.OnSelectionChanged_Lambda([OnSelectionChanged = InArgs._OnSelectionChanged](auto, auto){ OnSelectionChanged.ExecuteIfBound(); })
+					.OnExpansionChanged(this, &SReplicationTreeView::OnItemExpansionChanged)
 					.SelectionMode(InArgs._SelectionMode)
 					.AllowOverscroll(EAllowOverscroll::No)
 					.HeaderRow(CreateHeaderRow(InArgs))
@@ -619,6 +639,10 @@ namespace UE::ConcertSharedSlate
 			TreeView->RequestScrollIntoView(SelectedItems[0]);
 		}
 
+		// Update expansion after items have potentially changed
+		CleanseItemMetaData();
+		ReapplyExpansionStates();
+
 		bFilterChanged = false;
 		RequestResort();
 		TreeView->RequestListRefresh();
@@ -652,6 +676,55 @@ namespace UE::ConcertSharedSlate
 			&& (!CustomFilterDelegate.IsBound() || CustomFilterDelegate.Execute(Item));
 	}
 
+	template <typename TItemType>
+	void SReplicationTreeView<TItemType>::CleanseItemMetaData()
+	{
+		TMap<TSharedPtr<TItemType>, FItemMetaData> NewItemMetaData;
+		for (const TSharedPtr<TItemType>& Item : *AllRootItems)
+		{
+			NewItemMetaData.Add(Item, ItemMetaData.FindOrAdd(Item));
+		}
+		ItemMetaData = MoveTemp(NewItemMetaData);
+	}
+
+	template <typename TItemType>
+	void SReplicationTreeView<TItemType>::ReapplyExpansionStates()
+	{
+		// While searching, expand all items
+		bForceParentItemsExpanded = !SearchTextFilter->GetRawFilterText().IsEmpty();
+		
+		for (const TPair<TSharedPtr<TItemType>, FItemMetaData> MetaDataPair : ItemMetaData)
+		{
+			const TSharedPtr<TItemType>& Item = MetaDataPair.Key;
+			TreeView->SetItemExpansion(Item, MetaDataPair.Value.bIsExpanded || bForceParentItemsExpanded);
+		}
+	}
+
+	template <typename TItemType>
+	void SReplicationTreeView<TItemType>::OnItemExpansionChanged(TSharedPtr<TItemType> Item, bool bIsExpanded)
+	{
+		// We don't want to be overriding the cached state while everything is expanded
+		if (bForceParentItemsExpanded)
+		{
+			return;
+		}
+		
+		ItemMetaData.FindOrAdd(Item).bIsExpanded = bIsExpanded;
+		if (!bIsExpanded || !OnGetChildrenDelegate.IsBound())
+		{
+			return;
+		}
+		
+		// Also expand its children
+		OnGetChildrenDelegate.Execute(Item, [this](const TSharedPtr<TItemType>& Child)
+		{
+			if (const FItemMetaData* ChildMetaData = ItemMetaData.Find(Child))
+			{
+				TreeView->SetItemExpansion(Child, ChildMetaData->bIsExpanded);
+			}
+		});
+	}
+
 	template <typename TItemType>
 	void SReplicationTreeView<TItemType>::Resort()
 	{
