Adjust replication column content:
- Property Label column now displays names of properties instead of their property paths.  This makes it more consistent how the properties are displayed in details panel (e.g. bools are no longer displayed as bName but Name).
- The Type columns now have padding of 8 from left to align column content in a vertical line with column header

diff --git a/Engine/Plugins/Developer/Concert/ConcertUI/ConcertSharedSlate/Source/ConcertSharedSlate/Private/Replication/Editor/View/DisplayUtils.cpp b/Engine/Plugins/Developer/Concert/ConcertUI/ConcertSharedSlate/Source/ConcertSharedSlate/Private/Replication/Editor/View/DisplayUtils.cpp
--- a/Engine/Plugins/Developer/Concert/ConcertUI/ConcertSharedSlate/Source/ConcertSharedSlate/Private/Replication/Editor/View/DisplayUtils.cpp
+++ b/Engine/Plugins/Developer/Concert/ConcertUI/ConcertSharedSlate/Source/ConcertSharedSlate/Private/Replication/Editor/View/DisplayUtils.cpp
@@ -109,13 +109,21 @@ namespace UE::ConcertSharedSlate::DisplayUtils
 		return FSlateIconFinder::FindIconForClass(Object.GetClass());
 	}
 	
-	FText GetPropertyDisplayText(const FConcertPropertyChain& Property)
+	FText GetPropertyDisplayText(const FConcertPropertyChain& Property, UStruct* Class /*=nullptr*/)
 	{
-		return FText::FromString(GetPropertyDisplayString(Property));
+		// Class will (likely) be valid on editor builds but null on server
+		const FProperty* ResolvedProperty = Class ? Property.ResolveProperty(*Class) : nullptr;
+		return ResolvedProperty
+			? ResolvedProperty->GetDisplayNameText()
+			: FText::FromString(GetPropertyDisplayString(Property, Class));
     }
 	
-	FString GetPropertyDisplayString(const FConcertPropertyChain& Property)
+	FString GetPropertyDisplayString(const FConcertPropertyChain& Property, UStruct* Class /*=nullptr*/)
 	{
-		return Property.ToString(FConcertPropertyChain::EToStringMethod::LeafProperty);
+		// Class will (likely) be valid on editor builds but null on server
+		const FProperty* ResolvedProperty = Class ? Property.ResolveProperty(*Class) : nullptr;
+		return ResolvedProperty
+			? ResolvedProperty->GetDisplayNameText().ToString()
+			: Property.ToString(FConcertPropertyChain::EToStringMethod::LeafProperty);
 	}
 }
diff --git a/Engine/Plugins/Developer/Concert/ConcertUI/ConcertSharedSlate/Source/ConcertSharedSlate/Private/Replication/Editor/View/DisplayUtils.h b/Engine/Plugins/Developer/Concert/ConcertUI/ConcertSharedSlate/Source/ConcertSharedSlate/Private/Replication/Editor/View/DisplayUtils.h
--- a/Engine/Plugins/Developer/Concert/ConcertUI/ConcertSharedSlate/Source/ConcertSharedSlate/Private/Replication/Editor/View/DisplayUtils.h
+++ b/Engine/Plugins/Developer/Concert/ConcertUI/ConcertSharedSlate/Source/ConcertSharedSlate/Private/Replication/Editor/View/DisplayUtils.h
@@ -2,11 +2,12 @@
 
 #pragma once
 
-#include "Templates/SharedPointer.h"
 #include "Internationalization/Text.h"
+#include "Templates/SharedPointer.h"
 
 class FString;
 class UObject;
+class UStruct;
 
 struct FConcertPropertyChain;
 struct FSlateIcon;
@@ -38,8 +39,8 @@ namespace UE::ConcertSharedSlate::DisplayUtils
 	/** @return The icon to use for this object */
 	FSlateIcon GetObjectIcon(UObject& Object);
 
-	/** @return The text to use for displaying this property's name. */
-	FText GetPropertyDisplayText(const FConcertPropertyChain& Property);
-	/** @return More lightweight version of GetPropertyDisplayString which does not construct any FText. */
-	FString GetPropertyDisplayString(const FConcertPropertyChain& Property);
+	/** @return The text to use for displaying this property's name. Uses Class to determine class name if available. */
+	FText GetPropertyDisplayText(const FConcertPropertyChain& Property, UStruct* Class = nullptr);
+	/** @return More lightweight version of GetPropertyDisplayString which does not construct any FText. Uses Class to determine class name if available. */
+	FString GetPropertyDisplayString(const FConcertPropertyChain& Property, UStruct* Class = nullptr);
 }
diff --git a/Engine/Plugins/Developer/Concert/ConcertUI/ConcertSharedSlate/Source/ConcertSharedSlate/Private/Replication/Editor/View/SelectionViewerColumns.cpp b/Engine/Plugins/Developer/Concert/ConcertUI/ConcertSharedSlate/Source/ConcertSharedSlate/Private/Replication/Editor/View/SelectionViewerColumns.cpp
--- a/Engine/Plugins/Developer/Concert/ConcertUI/ConcertSharedSlate/Source/ConcertSharedSlate/Private/Replication/Editor/View/SelectionViewerColumns.cpp
+++ b/Engine/Plugins/Developer/Concert/ConcertUI/ConcertSharedSlate/Source/ConcertSharedSlate/Private/Replication/Editor/View/SelectionViewerColumns.cpp
@@ -14,9 +14,10 @@
 #include "Replication/Editor/View/Column/IObjectTreeColumn.h"
 #include "Textures/SlateIcon.h"
 #include "Widgets/Images/SImage.h"
+#include "Widgets/Input/SCheckBox.h"
+#include "Widgets/Layout/SBox.h"
 #include "Widgets/SBoxPanel.h"
 #include "Widgets/Text/STextBlock.h"
-#include "Widgets/Input/SCheckBox.h"
 #include "Styling/AppStyle.h"
 
 #define LOCTEXT_NAMESPACE "ReplicationObjectColumns"
@@ -126,9 +127,13 @@ namespace UE::ConcertSharedSlate::ReplicationColumns::TopLevel
 			
 			virtual TSharedRef<SWidget> GenerateColumnWidget(const FBuildArgs& InArgs) override
 			{
-				return SNew(STextBlock)
-					.HighlightText(TAttribute<FText>::CreateLambda([HighlightText = InArgs.HighlightText](){ return *HighlightText; }))
-					.Text(DisplayUtils::GetObjectTypeText(*Model, InArgs.RowItem.RowData.GetObjectPath()));
+				return SNew(SBox)
+					.Padding(8, 0, 0, 0) // So the type name text is aligned with the header column text
+					[
+						SNew(STextBlock)
+						.HighlightText(TAttribute<FText>::CreateLambda([HighlightText = InArgs.HighlightText](){ return *HighlightText; }))
+						.Text(DisplayUtils::GetObjectTypeText(*Model, InArgs.RowItem.RowData.GetObjectPath()))
+					];
 			}
 			
 			virtual void PopulateSearchString(const FObjectTreeRowContext& InItem, TArray<FString>& InOutSearchStrings) const override
@@ -185,21 +190,36 @@ namespace UE::ConcertSharedSlate::ReplicationColumns::Property
 			{
 				const FReplicatedPropertyData& PropertyData = InArgs.RowItem.RowData;
 				return SNew(STextBlock)
-						.Font(FAppStyle::GetFontStyle("PropertyWindow.NormalFont"))
-						.HighlightText(TAttribute<FText>::CreateLambda([HighlightText = InArgs.HighlightText](){ return *HighlightText; }))
-						.Text(DisplayUtils::GetPropertyDisplayText(PropertyData.GetProperty()));
+					.Font(FAppStyle::GetFontStyle("PropertyWindow.NormalFont"))
+					.HighlightText(TAttribute<FText>::CreateLambda([HighlightText = InArgs.HighlightText](){ return *HighlightText; }))
+					.Text(DisplayUtils::GetPropertyDisplayText(PropertyData.GetProperty(), ResolveOrLoadClass(PropertyData)));
 			}
 			
 			virtual void PopulateSearchString(const FPropertyTreeRowContext& InItem, TArray<FString>& InOutSearchStrings) const override
 			{
-				InOutSearchStrings.Add(DisplayUtils::GetPropertyDisplayText(InItem.RowData.GetProperty()).ToString());
+				const FReplicatedPropertyData& PropertyData = InItem.RowData;
+				FString DisplayString = DisplayUtils::GetPropertyDisplayString(InItem.RowData.GetProperty(), ResolveOrLoadClass(PropertyData));
+				InOutSearchStrings.Emplace(MoveTemp(DisplayString));
 			}
 			
 			virtual bool CanBeSorted() const override { return true; } 
 			virtual bool IsLessThan(const FPropertyTreeRowContext& Left, const FPropertyTreeRowContext& Right) const override
 			{
-				return DisplayUtils::GetPropertyDisplayString(Left.RowData.GetProperty())
-					< DisplayUtils::GetPropertyDisplayString(Right.RowData.GetProperty());
+				return DisplayUtils::GetPropertyDisplayString(Left.RowData.GetProperty(), ResolveOrLoadClass(Left.RowData))
+					< DisplayUtils::GetPropertyDisplayString(Right.RowData.GetProperty(), ResolveOrLoadClass(Right.RowData));
+			}
+
+		private:
+
+			static UClass* ResolveOrLoadClass(const FReplicatedPropertyData& PropertyData)
+			{
+#if WITH_EDITOR
+				// On editor there may be Blueprints that need loading...
+				return PropertyData.GetOwningClass().TryLoadClass<UObject>();
+#else
+				// ... but on the server everything should be native C++ and hence already loaded
+				return PropertyData.GetOwningClass().ResolveClass();
+#endif
 			}
 		};
 		
@@ -229,10 +249,14 @@ namespace UE::ConcertSharedSlate::ReplicationColumns::Property
 			virtual TSharedRef<SWidget> GenerateColumnWidget(const FBuildArgs& InArgs) override
 			{
 				const FReplicatedPropertyData& PropertyData = InArgs.RowItem.RowData;
-				return SNew(STextBlock)
+				return SNew(SBox)
+					.Padding(8, 0, 0, 0) // So the type name text is aligned with the header column text
+					[
+						SNew(STextBlock)
 						.Font(FAppStyle::GetFontStyle("PropertyWindow.NormalFont"))
 						.HighlightText(TAttribute<FText>::CreateLambda([HighlightText = InArgs.HighlightText](){ return *HighlightText; }))
-						.Text(GetDisplayText(InArgs.RowItem.RowData));
+						.Text(GetDisplayText(InArgs.RowItem.RowData))
+					];
 			}
 			
 			virtual void PopulateSearchString(const FPropertyTreeRowContext& InItem, TArray<FString>& InOutSearchStrings) const override
@@ -253,7 +277,7 @@ namespace UE::ConcertSharedSlate::ReplicationColumns::Property
 				UClass* Class = Args.GetOwningClass().TryLoadClass<UObject>();
 				const FProperty* Property = Class ? ConcertSyncCore::PropertyChain::ResolveProperty(*Class, Args.GetProperty()) : nullptr;
 				return Property ? FText::FromString(Property->GetCPPType()) : LOCTEXT("Unknown", "Unknown");	
-			};
+			}
 		};
 		
 		return {
