Wrap FProperty::GetDisplayNameText with WITH_EDITORONLY_DATA

diff --git a/Engine/Plugins/Developer/Concert/ConcertUI/ConcertSharedSlate/Source/ConcertSharedSlate/Private/Replication/Editor/View/DisplayUtils.cpp b/Engine/Plugins/Developer/Concert/ConcertUI/ConcertSharedSlate/Source/ConcertSharedSlate/Private/Replication/Editor/View/DisplayUtils.cpp
--- a/Engine/Plugins/Developer/Concert/ConcertUI/ConcertSharedSlate/Source/ConcertSharedSlate/Private/Replication/Editor/View/DisplayUtils.cpp
+++ b/Engine/Plugins/Developer/Concert/ConcertUI/ConcertSharedSlate/Source/ConcertSharedSlate/Private/Replication/Editor/View/DisplayUtils.cpp
@@ -16,6 +16,10 @@
 #include "GameFramework/Actor.h"
 #endif
 
+#if WITH_EDITORONLY_DATA
+#include "UObject/UnrealType.h"
+#endif
+
 namespace UE::ConcertSharedSlate::DisplayUtils
 {
 	FText GetObjectDisplayText(const FSoftObjectPath& Object, IObjectNameModel* Model)
@@ -111,19 +115,29 @@ namespace UE::ConcertSharedSlate::DisplayUtils
 	
 	FText GetPropertyDisplayText(const FConcertPropertyChain& Property, UStruct* Class /*=nullptr*/)
 	{
+#if WITH_EDITORONLY_DATA
 		// Class will (likely) be valid on editor builds but null on server
 		const FProperty* ResolvedProperty = Class ? Property.ResolveProperty(*Class) : nullptr;
 		return ResolvedProperty
 			? ResolvedProperty->GetDisplayNameText()
-			: FText::FromString(GetPropertyDisplayString(Property, Class));
+			:
+#else
+		return 
+#endif
+		FText::FromString(GetPropertyDisplayString(Property, Class));
     }
 	
 	FString GetPropertyDisplayString(const FConcertPropertyChain& Property, UStruct* Class /*=nullptr*/)
 	{
+#if WITH_EDITORONLY_DATA
 		// Class will (likely) be valid on editor builds but null on server
 		const FProperty* ResolvedProperty = Class ? Property.ResolveProperty(*Class) : nullptr;
 		return ResolvedProperty
 			? ResolvedProperty->GetDisplayNameText().ToString()
-			: Property.ToString(FConcertPropertyChain::EToStringMethod::LeafProperty);
+			:
+#else
+		return 
+#endif
+		Property.ToString(FConcertPropertyChain::EToStringMethod::LeafProperty);
 	}
 }
