Refactor replication UI in preparation of matrix view:
- Introduced IPropertyAssignmentView which abstracts the concept of displaying properties for an object
- SPerObjectPropertyAssignment implements IPropertyAssignmentView and performs the old behaviour (display the properties of the object selected in top view)

diff --git a/Engine/Plugins/Developer/Concert/ConcertApp/MultiUserClient/Source/MultiUserClient/Private/Widgets/ActiveSession/Replication/Client/Multi/Columns/MultiStreamColumns_AssignProperties.cpp b/Engine/Plugins/Developer/Concert/ConcertApp/MultiUserClient/Source/MultiUserClient/Private/Widgets/ActiveSession/Replication/Client/Multi/Columns/MultiStreamColumns_AssignProperties.cpp
--- a/Engine/Plugins/Developer/Concert/ConcertApp/MultiUserClient/Source/MultiUserClient/Private/Widgets/ActiveSession/Replication/Client/Multi/Columns/MultiStreamColumns_AssignProperties.cpp
+++ b/Engine/Plugins/Developer/Concert/ConcertApp/MultiUserClient/Source/MultiUserClient/Private/Widgets/ActiveSession/Replication/Client/Multi/Columns/MultiStreamColumns_AssignProperties.cpp
@@ -106,7 +106,7 @@ namespace UE::MultiUserClient::MultiStreamColumns
 			
 			virtual TSharedRef<SWidget> GenerateColumnWidget(const FBuildArgs& InArgs) override
 			{
-				const TArray<FSoftObjectPath> DisplayedObjects = MultiStreamEditor.Get()->GetEditorBase().GetObjectsBeingPropertyEdited();
+				const TSet<FSoftObjectPath> DisplayedObjects = MultiStreamEditor.Get()->GetEditorBase().GetObjectsBeingPropertyEdited();
 				return SNew(SAssignPropertyComboBox, MultiStreamEditor.Get().ToSharedRef(), ConcertClient, ClientManager)
 					.DisplayedProperty(InArgs.RowItem.RowData.GetProperty())
 					.EditedObjects(DisplayedObjects)
@@ -132,7 +132,7 @@ namespace UE::MultiUserClient::MultiStreamColumns
 			virtual bool CanBeSorted() const override { return true; }
 			virtual bool IsLessThan(const FPropertyTreeRowContext& Left, const FPropertyTreeRowContext& Right) const override
 			{
-				const TArray<FSoftObjectPath> DisplayedObjects = MultiStreamEditor.Get()->GetEditorBase().GetObjectsBeingPropertyEdited();
+				const TSet<FSoftObjectPath> DisplayedObjects = MultiStreamEditor.Get()->GetEditorBase().GetObjectsBeingPropertyEdited();
 				const TOptional<FString> LeftClientDisplayString = SAssignPropertyComboBox::GetDisplayString(ConcertClient, ClientManager, Left.RowData.GetProperty(), DisplayedObjects);
 				const TOptional<FString> RightClientDisplayString = SAssignPropertyComboBox::GetDisplayString(ConcertClient, ClientManager, Right.RowData.GetProperty(), DisplayedObjects);
 			
diff --git a/Engine/Plugins/Developer/Concert/ConcertApp/MultiUserClient/Source/MultiUserClient/Private/Widgets/ActiveSession/Replication/Client/Multi/Columns/SAssignPropertyComboBox.cpp b/Engine/Plugins/Developer/Concert/ConcertApp/MultiUserClient/Source/MultiUserClient/Private/Widgets/ActiveSession/Replication/Client/Multi/Columns/SAssignPropertyComboBox.cpp
--- a/Engine/Plugins/Developer/Concert/ConcertApp/MultiUserClient/Source/MultiUserClient/Private/Widgets/ActiveSession/Replication/Client/Multi/Columns/SAssignPropertyComboBox.cpp
+++ b/Engine/Plugins/Developer/Concert/ConcertApp/MultiUserClient/Source/MultiUserClient/Private/Widgets/ActiveSession/Replication/Client/Multi/Columns/SAssignPropertyComboBox.cpp
@@ -29,7 +29,7 @@ namespace UE::MultiUserClient
 {
 	namespace AssignPropertyComboBox
 	{
-		TArray<FGuid> GetDisplayedClients(const FReplicationClientManager& ClientManager, const FConcertPropertyChain& DisplayedProperty, const TArray<FSoftObjectPath>& EditedObjects)
+		TArray<FGuid> GetDisplayedClients(const FReplicationClientManager& ClientManager, const FConcertPropertyChain& DisplayedProperty, const TSet<FSoftObjectPath>& EditedObjects)
 		{
 			TArray<FGuid> Clients;
 			ClientManager.ForEachClient([&DisplayedProperty, &EditedObjects, &Clients](const FReplicationClient& Client)
@@ -54,7 +54,7 @@ namespace UE::MultiUserClient
 		const TSharedRef<IConcertClient>& LocalConcertClient,
 		const FReplicationClientManager& ClientManager,
 		const FConcertPropertyChain& DisplayedProperty,
-		const TArray<FSoftObjectPath>& EditedObjects)
+		const TSet<FSoftObjectPath>& EditedObjects)
 	{
 		using SWidgetType = ConcertClientSharedSlate::SHorizontalClientList;
 		const TArray<FGuid> Clients = AssignPropertyComboBox::GetDisplayedClients(ClientManager, DisplayedProperty, EditedObjects);
@@ -324,11 +324,15 @@ namespace UE::MultiUserClient
 		bool bIsAssignedToAnyClient = false;
 		ClientManager->ForEachClient([this, &bIsAssignedToAnyClient](const FReplicationClient& Client)
 		{
-			for (int32 i = 0; !bIsAssignedToAnyClient && i < EditedObjects.Num(); ++i)
+			for (const FSoftObjectPath& EditedObject : EditedObjects)
 			{
-				const FSoftObjectPath& ObjectPath = EditedObjects[i];
+				if (bIsAssignedToAnyClient)
+				{
+					break;
+				}
+				
 				const TSharedRef<ConcertSharedSlate::IEditableReplicationStreamModel> Model = Client.GetClientEditModel();
-				const bool bHasProperty = Model->HasProperty(ObjectPath, Property);
+				const bool bHasProperty = Model->HasProperty(EditedObject, Property);
 				bIsAssignedToAnyClient |= bHasProperty;
 			}
 			
diff --git a/Engine/Plugins/Developer/Concert/ConcertApp/MultiUserClient/Source/MultiUserClient/Private/Widgets/ActiveSession/Replication/Client/Multi/Columns/SAssignPropertyComboBox.h b/Engine/Plugins/Developer/Concert/ConcertApp/MultiUserClient/Source/MultiUserClient/Private/Widgets/ActiveSession/Replication/Client/Multi/Columns/SAssignPropertyComboBox.h
--- a/Engine/Plugins/Developer/Concert/ConcertApp/MultiUserClient/Source/MultiUserClient/Private/Widgets/ActiveSession/Replication/Client/Multi/Columns/SAssignPropertyComboBox.h
+++ b/Engine/Plugins/Developer/Concert/ConcertApp/MultiUserClient/Source/MultiUserClient/Private/Widgets/ActiveSession/Replication/Client/Multi/Columns/SAssignPropertyComboBox.h
@@ -38,13 +38,13 @@ namespace UE::MultiUserClient
 			const TSharedRef<IConcertClient>& LocalConcertClient,
 			const FReplicationClientManager& ClientManager,
 			const FConcertPropertyChain& DisplayedProperty,
-			const TArray<FSoftObjectPath>& EditedObjects
+			const TSet<FSoftObjectPath>& EditedObjects
 			);
 		
 		SLATE_BEGIN_ARGS(SAssignPropertyComboBox)
 		{}
 			SLATE_ARGUMENT(FConcertPropertyChain, DisplayedProperty)
-			SLATE_ARGUMENT(TArray<FSoftObjectPath>, EditedObjects)
+			SLATE_ARGUMENT(TSet<FSoftObjectPath>, EditedObjects)
 			SLATE_ARGUMENT(TSharedPtr<FText>, HighlightText)
 			/** Called when the property assignment of any client(s) is changed by this widget. */
 			SLATE_EVENT(FOnPropertyAssignmentChanged, OnPropertyAssignmentChanged)
@@ -65,7 +65,7 @@ namespace UE::MultiUserClient
 		TSharedPtr<IConcertClient> ConcertClient;
 
 		/** The objects for which the property is being displayed. */
-		TArray<FSoftObjectPath> EditedObjects;
+		TSet<FSoftObjectPath> EditedObjects;
 		/** The property assigned to this column. */
 		FConcertPropertyChain Property;
 
diff --git a/Engine/Plugins/Developer/Concert/ConcertApp/MultiUserClient/Source/MultiUserClient/Private/Widgets/ActiveSession/Replication/Client/Multi/SMultiClientView.cpp b/Engine/Plugins/Developer/Concert/ConcertApp/MultiUserClient/Source/MultiUserClient/Private/Widgets/ActiveSession/Replication/Client/Multi/SMultiClientView.cpp
--- a/Engine/Plugins/Developer/Concert/ConcertApp/MultiUserClient/Source/MultiUserClient/Private/Widgets/ActiveSession/Replication/Client/Multi/SMultiClientView.cpp
+++ b/Engine/Plugins/Developer/Concert/ConcertApp/MultiUserClient/Source/MultiUserClient/Private/Widgets/ActiveSession/Replication/Client/Multi/SMultiClientView.cpp
@@ -78,6 +78,7 @@ namespace UE::MultiUserClient
 			const TSharedRef<IEditableReplicationStreamModel>& LocalStream = InClientManager.GetLocalClient().GetClientEditModel();
 			return StreamModel->GetEditableStreams().Contains(LocalStream) ? LocalStream.ToSharedPtr() : nullptr;
 		});
+		const TSharedRef<ConcertClientSharedSlate::FSelectPropertyFromUClassModel> PropertySourceModel = MakeShared<ConcertClientSharedSlate::FSelectPropertyFromUClassModel>();
 		
 		ConcertClientSharedSlate::FFilterablePropertyTreeViewParams TreeViewParams
 		{
@@ -89,20 +90,21 @@ namespace UE::MultiUserClient
 			}
 		};
 		TSharedRef<IPropertyTreeView> PropertyTreeView = CreateFilterablePropertyTreeView(MoveTemp(TreeViewParams));
+		TSharedRef<IPropertyAssignmentView> PropertyAssignmentView = CreatePerObjectAssignmentView({ .PropertyTreeView = PropertyTreeView, .PropertySource = PropertySourceModel });
 		
 		FCreateMultiStreamEditorParams Params
 		{
 			.MultiStreamModel = StreamModel.ToSharedRef(),
 			.ConsolidatedObjectModel = ConcertClientSharedSlate::CreateTransactionalStreamModel(),
 			.ObjectSource = MakeShared<ConcertClientSharedSlate::FActorSelectionSourceModel>(),
-			.PropertySource = MakeShared<ConcertClientSharedSlate::FSelectPropertyFromUClassModel>(),
+			.PropertySource = PropertySourceModel,
 			.GetAutoAssignToStreamDelegate = MoveTemp(GetAutoAssignTargetDelegate)
 		};
 		
 		ObjectHierarchy = ConcertClientSharedSlate::CreateObjectHierarchyForComponentHierarchy();
 		FCreateViewerParams ViewerParams
 		{
-			.PropertyTreeView = MoveTemp(PropertyTreeView),
+			.PropertyAssignmentView = MoveTemp(PropertyAssignmentView),
 			.ObjectHierarchy = ObjectHierarchy, // This makes actors have children in the top view
 			.NameModel = ConcertClientSharedSlate::CreateEditorObjectNameModel(), // This makes actors use their labels, and components use the names given in the BP editor
 			.OnExtendObjectsContextMenu = FExtendObjectMenu::CreateSP(this, &SMultiClientView::ExtendObjectContextMenu),
diff --git a/Engine/Plugins/Developer/Concert/ConcertApp/MultiUserClient/Source/MultiUserClient/Private/Widgets/ActiveSession/Replication/Joined/SReplicationJoinedView.h b/Engine/Plugins/Developer/Concert/ConcertApp/MultiUserClient/Source/MultiUserClient/Private/Widgets/ActiveSession/Replication/Joined/SReplicationJoinedView.h
--- a/Engine/Plugins/Developer/Concert/ConcertApp/MultiUserClient/Source/MultiUserClient/Private/Widgets/ActiveSession/Replication/Joined/SReplicationJoinedView.h
+++ b/Engine/Plugins/Developer/Concert/ConcertApp/MultiUserClient/Source/MultiUserClient/Private/Widgets/ActiveSession/Replication/Joined/SReplicationJoinedView.h
@@ -30,5 +30,12 @@ namespace UE::MultiUserClient
 			const TSharedRef<FMultiUserReplicationManager>& InReplicationManager,
 			const TSharedRef<IConcertSyncClient>& InClient
 			);
+
+	private:
+
+		/** The local client this widget is created for. */
+		TSharedPtr<IConcertSyncClient> Client;
+		/** Acts as the model of this view */
+		TSharedPtr<FMultiUserReplicationManager> ReplicationManager;
 	};
 }
\ No newline at end of file
diff --git a/Engine/Plugins/Developer/Concert/ConcertScripting/ConcertReplicationScripting/Source/ConcertReplicationScriptingEditor/Private/Widgets/SConcertPropertyChainPicker.cpp b/Engine/Plugins/Developer/Concert/ConcertScripting/ConcertReplicationScripting/Source/ConcertReplicationScriptingEditor/Private/Widgets/SConcertPropertyChainPicker.cpp
--- a/Engine/Plugins/Developer/Concert/ConcertScripting/ConcertReplicationScripting/Source/ConcertReplicationScriptingEditor/Private/Widgets/SConcertPropertyChainPicker.cpp
+++ b/Engine/Plugins/Developer/Concert/ConcertScripting/ConcertReplicationScripting/Source/ConcertReplicationScriptingEditor/Private/Widgets/SConcertPropertyChainPicker.cpp
@@ -133,7 +133,7 @@ namespace UE::ConcertReplicationScriptingEditor
 			});
 		}
 
-		// If not class is selected, the widget will display the NoItemsContent ("Select a class").
+		// If no class is selected, the widget will display the NoItemsContent ("Select a class").
 		TreeView->RefreshPropertyData(Properties, ClassPath);
 	}
 
@@ -141,7 +141,7 @@ namespace UE::ConcertReplicationScriptingEditor
 	{
 		OnSelectedPropertiesChangedDelegate.ExecuteIfBound(ConcertPropertyChain, bIsSelected);
 
-		// The checkbox state may have changed request a resort.
+		// The checkbox state may have changed so request a resort.
 		TreeView->RequestResortForColumn(PropertySelectionCheckboxColumnId);
 	}
 }
diff --git a/Engine/Plugins/Developer/Concert/ConcertUI/ConcertClientSharedSlate/Source/ConcertClientSharedSlate/Private/Replication/Editor/Model/Object/EditorObjectNameModel.cpp b/Engine/Plugins/Developer/Concert/ConcertUI/ConcertClientSharedSlate/Source/ConcertClientSharedSlate/Private/Replication/Editor/Model/Object/EditorObjectNameModel.cpp
--- a/Engine/Plugins/Developer/Concert/ConcertUI/ConcertClientSharedSlate/Source/ConcertClientSharedSlate/Private/Replication/Editor/Model/Object/EditorObjectNameModel.cpp
+++ b/Engine/Plugins/Developer/Concert/ConcertUI/ConcertClientSharedSlate/Source/ConcertClientSharedSlate/Private/Replication/Editor/Model/Object/EditorObjectNameModel.cpp
@@ -3,7 +3,7 @@
 #include "EditorObjectNameModel.h"
 
 #include "Internationalization/Text.h"
-#include "Replication/ObjectUtils.h"
+#include "Replication/Utils/ObjectUtils.h"
 
 #include "GameFramework/Actor.h"
 #include "SubobjectDataSubsystem.h"
diff --git a/Engine/Plugins/Developer/Concert/ConcertUI/ConcertClientSharedSlate/Source/ConcertClientSharedSlate/Private/Replication/Editor/View/PropertyTree/PropertyFilterBase.h b/Engine/Plugins/Developer/Concert/ConcertUI/ConcertClientSharedSlate/Source/ConcertClientSharedSlate/Private/Replication/Editor/View/PropertyTree/PropertyFilterBase.h
--- a/Engine/Plugins/Developer/Concert/ConcertUI/ConcertClientSharedSlate/Source/ConcertClientSharedSlate/Private/Replication/Editor/View/PropertyTree/PropertyFilterBase.h
+++ b/Engine/Plugins/Developer/Concert/ConcertUI/ConcertClientSharedSlate/Source/ConcertClientSharedSlate/Private/Replication/Editor/View/PropertyTree/PropertyFilterBase.h
@@ -6,7 +6,7 @@
 #include "Misc/IFilter.h"
 
 class FFieldClass;
-namespace UE::ConcertSharedSlate { class FReplicatedPropertyData; }
+namespace UE::ConcertSharedSlate { class FPropertyData; }
 
 namespace UE::ConcertClientSharedSlate
 {
@@ -16,14 +16,14 @@ namespace UE::ConcertClientSharedSlate
 	 * This filter inverts the results: all filters are supposed to be inverse filters, i.e. they should be run when
 	 * they are greyed out in the UI. When they are run, they should remove specific properties.
 	 *
-	 * Subclasses simply implement Matches, which figures out whether the property is contained in some FReplicatedPropertyData.
+	 * Subclasses simply implement Matches, which figures out whether the property is contained in some FPropertyData.
 	 */
-	class FPropertyFilterBase : public IFilter<const ConcertSharedSlate::FReplicatedPropertyData&>
+	class FPropertyFilterBase : public IFilter<const ConcertSharedSlate::FPropertyData&>
 	{
 	public:
 		
 		//~ Begin IFilter Interface
-		virtual bool PassesFilter(const ConcertSharedSlate::FReplicatedPropertyData& InItem) const final override
+		virtual bool PassesFilter(const ConcertSharedSlate::FPropertyData& InItem) const final override
 		{
 			return MatchesFilteredForProperty(InItem);
 		}
@@ -36,7 +36,7 @@ namespace UE::ConcertClientSharedSlate
 		FChangedEvent ChangedEventDelegate;
 
 		/** @return Whether this item contains the property this filter is looking for.*/
-		virtual bool MatchesFilteredForProperty(const ConcertSharedSlate::FReplicatedPropertyData& InItem) const = 0;
+		virtual bool MatchesFilteredForProperty(const ConcertSharedSlate::FPropertyData& InItem) const = 0;
 	};
 }
 
diff --git a/Engine/Plugins/Developer/Concert/ConcertUI/ConcertClientSharedSlate/Source/ConcertClientSharedSlate/Private/Replication/Editor/View/PropertyTree/PropertyFilter_ByPropertyType.cpp b/Engine/Plugins/Developer/Concert/ConcertUI/ConcertClientSharedSlate/Source/ConcertClientSharedSlate/Private/Replication/Editor/View/PropertyTree/PropertyFilter_ByPropertyType.cpp
--- a/Engine/Plugins/Developer/Concert/ConcertUI/ConcertClientSharedSlate/Source/ConcertClientSharedSlate/Private/Replication/Editor/View/PropertyTree/PropertyFilter_ByPropertyType.cpp
+++ b/Engine/Plugins/Developer/Concert/ConcertUI/ConcertClientSharedSlate/Source/ConcertClientSharedSlate/Private/Replication/Editor/View/PropertyTree/PropertyFilter_ByPropertyType.cpp
@@ -3,13 +3,13 @@
 #include "PropertyFilter_ByPropertyType.h"
 
 #include "Replication/PropertyChainUtils.h"
-#include "Replication/Editor/Model/ReplicatedPropertyData.h"
+#include "Replication/Editor/Model/Data/PropertyData.h"
 
 #include "UObject/UnrealType.h"
 
 namespace UE::ConcertClientSharedSlate
 {
-	bool FPropertyFilter_ByPropertyType::MatchesFilteredForProperty(const ConcertSharedSlate::FReplicatedPropertyData& InItem) const
+	bool FPropertyFilter_ByPropertyType::MatchesFilteredForProperty(const ConcertSharedSlate::FPropertyData& InItem) const
 	{
 		UClass* Class = InItem.GetOwningClass().TryLoadClass<UObject>();
 		if (!Class)
diff --git a/Engine/Plugins/Developer/Concert/ConcertUI/ConcertClientSharedSlate/Source/ConcertClientSharedSlate/Private/Replication/Editor/View/PropertyTree/PropertyFilter_ByPropertyType.h b/Engine/Plugins/Developer/Concert/ConcertUI/ConcertClientSharedSlate/Source/ConcertClientSharedSlate/Private/Replication/Editor/View/PropertyTree/PropertyFilter_ByPropertyType.h
--- a/Engine/Plugins/Developer/Concert/ConcertUI/ConcertClientSharedSlate/Source/ConcertClientSharedSlate/Private/Replication/Editor/View/PropertyTree/PropertyFilter_ByPropertyType.h
+++ b/Engine/Plugins/Developer/Concert/ConcertUI/ConcertClientSharedSlate/Source/ConcertClientSharedSlate/Private/Replication/Editor/View/PropertyTree/PropertyFilter_ByPropertyType.h
@@ -8,7 +8,7 @@
 
 class FFieldClass;
 
-namespace UE::ConcertSharedSlate { class FReplicatedPropertyData; }
+namespace UE::ConcertSharedSlate { class FPropertyData; }
 
 namespace UE::ConcertClientSharedSlate
 {
@@ -29,7 +29,7 @@ namespace UE::ConcertClientSharedSlate
 		FChangedEvent ChangedEventDelegate;
 
 		//~ Begin FPropertyFilterBase Interface
-		virtual bool MatchesFilteredForProperty(const ConcertSharedSlate::FReplicatedPropertyData& InItem) const override;
+		virtual bool MatchesFilteredForProperty(const ConcertSharedSlate::FPropertyData& InItem) const override;
 		//~ End FPropertyFilterBase Interface
 	};
 }
diff --git a/Engine/Plugins/Developer/Concert/ConcertUI/ConcertClientSharedSlate/Source/ConcertClientSharedSlate/Private/Replication/Editor/View/PropertyTree/PropertyFrontendFilter.h b/Engine/Plugins/Developer/Concert/ConcertUI/ConcertClientSharedSlate/Source/ConcertClientSharedSlate/Private/Replication/Editor/View/PropertyTree/PropertyFrontendFilter.h
--- a/Engine/Plugins/Developer/Concert/ConcertUI/ConcertClientSharedSlate/Source/ConcertClientSharedSlate/Private/Replication/Editor/View/PropertyTree/PropertyFrontendFilter.h
+++ b/Engine/Plugins/Developer/Concert/ConcertUI/ConcertClientSharedSlate/Source/ConcertClientSharedSlate/Private/Replication/Editor/View/PropertyTree/PropertyFrontendFilter.h
@@ -5,12 +5,12 @@
 #include "Filters/FilterBase.h"
 #include "Internationalization/Text.h"
 
-namespace UE::ConcertSharedSlate { class FReplicatedPropertyData; }
+namespace UE::ConcertSharedSlate { class FPropertyData; }
 
 namespace UE::ConcertClientSharedSlate
 {
 	/** Base UI filter for filtering properties. Gets rid of all the extra functionality in FFilterBase we do not need. */
-	class FPropertyFrontendFilter : public FFilterBase<const ConcertSharedSlate::FReplicatedPropertyData&>
+	class FPropertyFrontendFilter : public FFilterBase<const ConcertSharedSlate::FPropertyData&>
 	{
 	public:
 
@@ -39,7 +39,7 @@ namespace UE::ConcertClientSharedSlate
 		FText Tooltip;
 	};
 
-	/** Inlines a IFilter<const FReplicatedPropertyData&> at compile time to avoid an additional TSharedPtr<IFilter<const FReplicatedPropertyData&>>. */
+	/** Inlines a IFilter<const FPropertyData&> at compile time to avoid an additional TSharedPtr<IFilter<const FPropertyData&>>. */
 	template<typename TFilterType>
 	class TPropertyFrontendFilter : public FPropertyFrontendFilter
 	{
@@ -67,7 +67,7 @@ namespace UE::ConcertClientSharedSlate
 			});
 		}
 		
-		virtual bool PassesFilter(const ConcertSharedSlate::FReplicatedPropertyData& InItem) const override
+		virtual bool PassesFilter(const ConcertSharedSlate::FPropertyData& InItem) const override
 		{
 			return FilterImplementation.PassesFilter(InItem);
 		}
diff --git a/Engine/Plugins/Developer/Concert/ConcertUI/ConcertClientSharedSlate/Source/ConcertClientSharedSlate/Private/Replication/Editor/View/PropertyTree/SFilteredPropertyTreeView.cpp b/Engine/Plugins/Developer/Concert/ConcertUI/ConcertClientSharedSlate/Source/ConcertClientSharedSlate/Private/Replication/Editor/View/PropertyTree/SFilteredPropertyTreeView.cpp
--- a/Engine/Plugins/Developer/Concert/ConcertUI/ConcertClientSharedSlate/Source/ConcertClientSharedSlate/Private/Replication/Editor/View/PropertyTree/SFilteredPropertyTreeView.cpp
+++ b/Engine/Plugins/Developer/Concert/ConcertUI/ConcertClientSharedSlate/Source/ConcertClientSharedSlate/Private/Replication/Editor/View/PropertyTree/SFilteredPropertyTreeView.cpp
@@ -7,7 +7,7 @@
 #include "PropertyFrontendFilter.h"
 #include "Replication/ClientReplicationWidgetFactories.h"
 #include "Replication/ReplicationWidgetFactories.h"
-#include "Replication/Editor/Model/ReplicatedPropertyData.h"
+#include "Replication/Editor/Model/Data/PropertyData.h"
 
 #include "UObject/UnrealType.h"
 
@@ -16,9 +16,9 @@
 namespace UE::ConcertClientSharedSlate
 {
 	/** Exposes SetFrontendFilterActive so we can manually enable the default filters */
-	class SReplicationFilterBar : public SBasicFilterBar<const ConcertSharedSlate::FReplicatedPropertyData&>
+	class SReplicationFilterBar : public SBasicFilterBar<const ConcertSharedSlate::FPropertyData&>
 	{
-		using Super = SBasicFilterBar<const ConcertSharedSlate::FReplicatedPropertyData&>;
+		using Super = SBasicFilterBar<const ConcertSharedSlate::FPropertyData&>;
 	public:
 
 		SLATE_BEGIN_ARGS(SReplicationFilterBar)
@@ -65,7 +65,7 @@ namespace UE::ConcertClientSharedSlate
 			.PrimaryPropertySort = Params.PrimaryPropertySort,
 			.SecondaryPropertySort = Params.SecondaryPropertySort
 		};
-		TreeViewParams.LeftOfPropertySearchBar.Widget = SBasicFilterBar<const FReplicatedPropertyData&>::MakeAddFilterButton(FilterBar.ToSharedRef());
+		TreeViewParams.LeftOfPropertySearchBar.Widget = SBasicFilterBar<const FPropertyData&>::MakeAddFilterButton(FilterBar.ToSharedRef());
 		TreeViewParams.RowBelowSearchBar.Widget = FilterBar.ToSharedRef();
 		TreeViewParams.NoItemsContent.Widget = SNew(STextBlock).Text(LOCTEXT("AllFitlered", "All properties filtered."));
 			
@@ -139,16 +139,16 @@ namespace UE::ConcertClientSharedSlate
 		return Result;
 	}
 
-	ConcertSharedSlate::EFilterResult SFilteredPropertyTreeView::PassesFilters(const ConcertSharedSlate::FReplicatedPropertyData& ReplicatedPropertyData) const
+	ConcertSharedSlate::EFilterResult SFilteredPropertyTreeView::PassesFilters(const ConcertSharedSlate::FPropertyData& ReplicatedPropertyData) const
 	{
 		const bool bPassesFilter = FilterBar->GetAllActiveFilters()->Num() == 0 // Return all items when none enabled
 			|| PassesAnyFilters(ReplicatedPropertyData);
 		return bPassesFilter ? ConcertSharedSlate::EFilterResult::PassesFilter : ConcertSharedSlate::EFilterResult::DoesNotPassFilter;
 	}
 
-	bool SFilteredPropertyTreeView::PassesAnyFilters(const ConcertSharedSlate::FReplicatedPropertyData& ReplicatedPropertyData) const
+	bool SFilteredPropertyTreeView::PassesAnyFilters(const ConcertSharedSlate::FPropertyData& ReplicatedPropertyData) const
 	{
-		TSharedPtr<TFilterCollection<const ConcertSharedSlate::FReplicatedPropertyData&>> FilterCollection = FilterBar->GetAllActiveFilters();
+		TSharedPtr<TFilterCollection<const ConcertSharedSlate::FPropertyData&>> FilterCollection = FilterBar->GetAllActiveFilters();
 		for (int32 Index = 0; Index < FilterCollection->Num(); Index++)
 		{
 			if (FilterCollection->GetFilterAtIndex(Index)->PassesFilter(ReplicatedPropertyData))
diff --git a/Engine/Plugins/Developer/Concert/ConcertUI/ConcertClientSharedSlate/Source/ConcertClientSharedSlate/Private/Replication/Editor/View/PropertyTree/SFilteredPropertyTreeView.h b/Engine/Plugins/Developer/Concert/ConcertUI/ConcertClientSharedSlate/Source/ConcertClientSharedSlate/Private/Replication/Editor/View/PropertyTree/SFilteredPropertyTreeView.h
--- a/Engine/Plugins/Developer/Concert/ConcertUI/ConcertClientSharedSlate/Source/ConcertClientSharedSlate/Private/Replication/Editor/View/PropertyTree/SFilteredPropertyTreeView.h
+++ b/Engine/Plugins/Developer/Concert/ConcertUI/ConcertClientSharedSlate/Source/ConcertClientSharedSlate/Private/Replication/Editor/View/PropertyTree/SFilteredPropertyTreeView.h
@@ -11,7 +11,7 @@
 
 namespace UE::ConcertSharedSlate
 {
-	class FReplicatedPropertyData;
+	class FPropertyData;
 	enum class EFilterResult : uint8;
 }
 
@@ -27,7 +27,7 @@ namespace UE::ConcertClientSharedSlate
 	{
 	public:
 		
-		using FFilterRef = TSharedRef<FFilterBase<const ConcertSharedSlate::FReplicatedPropertyData&>>;
+		using FFilterRef = TSharedRef<FFilterBase<const ConcertSharedSlate::FPropertyData&>>;
 		
 		SLATE_BEGIN_ARGS(SFilteredPropertyTreeView)
 		{}
@@ -62,8 +62,8 @@ namespace UE::ConcertClientSharedSlate
 		FBuildFilterBarResult BuildFilterBar();
 
 		/** Runs all filters through this item */
-		ConcertSharedSlate::EFilterResult PassesFilters(const ConcertSharedSlate::FReplicatedPropertyData& ReplicatedPropertyData) const;
-		bool PassesAnyFilters(const ConcertSharedSlate::FReplicatedPropertyData& ReplicatedPropertyData) const;
+		ConcertSharedSlate::EFilterResult PassesFilters(const ConcertSharedSlate::FPropertyData& ReplicatedPropertyData) const;
+		bool PassesAnyFilters(const ConcertSharedSlate::FPropertyData& ReplicatedPropertyData) const;
 	};
 }
 
diff --git a/Engine/Plugins/Developer/Concert/ConcertUI/ConcertSharedSlate/Source/ConcertSharedSlate/Private/Replication/Editor/Utils/DisplayUtils.cpp b/Engine/Plugins/Developer/Concert/ConcertUI/ConcertSharedSlate/Source/ConcertSharedSlate/Private/Replication/Editor/Utils/DisplayUtils.cpp
new file mode 100644
--- /dev/null
+++ b/Engine/Plugins/Developer/Concert/ConcertUI/ConcertSharedSlate/Source/ConcertSharedSlate/Private/Replication/Editor/Utils/DisplayUtils.cpp
@@ -0,0 +1,142 @@
+// Copyright Epic Games, Inc. All Rights Reserved.
+
+#include "Replication/Editor/Utils/DisplayUtils.h"
+
+#include "ConcertLogGlobal.h"
+#include "Replication/Utils/ObjectUtils.h"
+#include "Replication/Editor/Model/IReplicationStreamModel.h"
+#include "Replication/Editor/Model/Object/IObjectNameModel.h"
+
+#include "Styling/SlateIconFinder.h"
+
+#if WITH_EDITOR
+#include "AssetRegistry/AssetData.h"
+#include "AssetRegistry/AssetRegistryModule.h"
+#include "Blueprint/BlueprintSupport.h"
+#include "GameFramework/Actor.h"
+#endif
+
+#if WITH_EDITORONLY_DATA
+#include "UObject/UnrealType.h"
+#endif
+
+namespace UE::ConcertSharedSlate::DisplayUtils
+{
+	FText GetObjectDisplayText(const FSoftObjectPath& Object, IObjectNameModel* Model)
+	{
+		if (Model)
+		{
+			return Model->GetObjectDisplayName(Object);
+		}
+		
+#if WITH_EDITOR
+		// Important! The object may not be loaded, yet. So only resolve, do not TryLoad. This could be if the asset is using a level that was not opened.
+		if (const UObject* LoadedObject = Object.ResolveObject())
+		{
+			return FText::FromString(GetObjectDisplayString(*LoadedObject));
+		}
+#endif
+
+		return ExtractObjectDisplayTextFromPath(Object);
+	}
+	
+	FText ExtractObjectDisplayTextFromPath(const FSoftObjectPath& Object)
+	{
+		return FText::FromString(ObjectUtils::ExtractObjectDisplayStringFromPath(Object));
+	}
+
+	FString GetObjectDisplayString(const UObject& Object)
+	{
+		// In editor builds we have access to actor label so display them (that's also what's displayed in the scene outliner)
+#if WITH_EDITOR 
+		if (const AActor* AsActor = Cast<AActor>(&Object))
+		{
+			return AsActor->GetActorLabel();
+		}
+#endif
+		// If we're not in the editor (or the object is not an actor) then whatever the name is in the soft object path will be good enough.
+		return Object.GetName();
+	}
+
+	FText GetObjectTypeText(const FSoftClassPath& ClassPath)
+	{
+#if WITH_EDITOR
+		if (ClassPath.IsValid())
+		{
+			// For C++ classes or loaded Blueprints
+			if (UClass* Class = ClassPath.ResolveClass())
+			{
+				return Class->GetDisplayNameText();
+			}
+
+			// Assuming it is a Blueprint, try to get its display name from the registry tags
+			const FAssetRegistryModule& AssetRegistryModule = FModuleManager::LoadModuleChecked<FAssetRegistryModule>(TEXT("AssetRegistry"));
+			const IAssetRegistry& AssetRegistry = AssetRegistryModule.Get();
+			
+			const FAssetData BlueprintData = AssetRegistry.GetAssetByObjectPath(ClassPath);
+			if (BlueprintData.IsValid())
+			{
+				FString DisplayName;
+				return BlueprintData.GetTagValue(FBlueprintTags::BlueprintDisplayName, DisplayName)
+					? FText::FromString(DisplayName)
+					: FText::GetEmpty();
+			}
+
+			// TODO: Check redirectors
+			// It was not a Blueprint - the class was most likely removed
+			UE_LOG(LogConcert, Error, TEXT("Failed to get display name for unknown class \"%s\""), *ClassPath.ToString());
+		}
+		
+		return FText::FromName(NAME_None);
+#else
+		// On server the type info is not available. We can just display what it says in the path.
+		return FText::FromString(ClassPath.GetAssetName());
+#endif
+
+	}
+
+	FSlateIcon GetObjectIcon(const FSoftClassPath& ClassPath)
+	{
+		if (ClassPath.IsValid())
+		{
+			UClass* Class = ClassPath.ResolveClass();
+			return Class
+				? FSlateIconFinder::FindIconForClass(Class)
+				: FSlateIcon();
+		}
+		return FSlateIcon();
+	}
+
+	FSlateIcon GetObjectIcon(UObject& Object)
+	{
+		return FSlateIconFinder::FindIconForClass(Object.GetClass());
+	}
+	
+	FText GetPropertyDisplayText(const FConcertPropertyChain& Property, UStruct* Class /*=nullptr*/)
+	{
+#if WITH_EDITORONLY_DATA
+		// Class will (likely) be valid on editor builds but null on server
+		const FProperty* ResolvedProperty = Class ? Property.ResolveProperty(*Class) : nullptr;
+		return ResolvedProperty
+			? ResolvedProperty->GetDisplayNameText()
+			:
+#else
+		return 
+#endif
+		FText::FromString(GetPropertyDisplayString(Property, Class));
+    }
+	
+	FString GetPropertyDisplayString(const FConcertPropertyChain& Property, UStruct* Class /*=nullptr*/)
+	{
+#if WITH_EDITORONLY_DATA
+		// Class will (likely) be valid on editor builds but null on server
+		const FProperty* ResolvedProperty = Class ? Property.ResolveProperty(*Class) : nullptr;
+		return ResolvedProperty
+			? ResolvedProperty->GetDisplayNameText().ToString()
+			:
+#else
+		return 
+#endif
+		Property.ToString(FConcertPropertyChain::EToStringMethod::LeafProperty);
+	}
+}
diff --git a/Engine/Plugins/Developer/Concert/ConcertUI/ConcertSharedSlate/Source/ConcertSharedSlate/Private/Replication/Editor/Utils/DisplayUtils.h b/Engine/Plugins/Developer/Concert/ConcertUI/ConcertSharedSlate/Source/ConcertSharedSlate/Private/Replication/Editor/Utils/DisplayUtils.h
new file mode 100644
--- /dev/null
+++ b/Engine/Plugins/Developer/Concert/ConcertUI/ConcertSharedSlate/Source/ConcertSharedSlate/Private/Replication/Editor/Utils/DisplayUtils.h
@@ -0,0 +1,47 @@
+// Copyright Epic Games, Inc. All Rights Reserved.
+
+#pragma once
+
+#include "Internationalization/Text.h"
+#include "Templates/SharedPointer.h"
+
+class FString;
+class UObject;
+class UStruct;
+
+struct FConcertPropertyChain;
+struct FSlateIcon;
+struct FSoftClassPath;
+struct FSoftObjectPath;
+
+namespace UE::ConcertSharedSlate
+{
+	class IObjectNameModel;
+	class IReplicationStreamModel;
+}
+
+namespace UE::ConcertSharedSlate::DisplayUtils
+{
+	/** Unified version that reads from IObjectNameModel or defaults to extracting the name from the path. */
+	FText GetObjectDisplayText(const FSoftObjectPath& Object, IObjectNameModel* Model = nullptr);
+	inline FText GetObjectDisplayText(const FSoftObjectPath& Object, const TSharedPtr<IObjectNameModel>& Model = nullptr) { return GetObjectDisplayText(Object, Model.Get()); }
+	
+	/** @return The text to use for displaying this object's name */
+	FText ExtractObjectDisplayTextFromPath(const FSoftObjectPath& Object);
+	
+	/** @return More lightweight version of GetObjectDisplayText which does not construct any FText. */
+	FString GetObjectDisplayString(const UObject& Object);
+	
+	/** @return The text to use for displaying this object's type */
+	FText GetObjectTypeText(const FSoftClassPath& ClassPath);
+	
+	/** @return The icon to use for this object */
+	FSlateIcon GetObjectIcon(const FSoftClassPath& ClassPath);
+	/** @return The icon to use for this object */
+	FSlateIcon GetObjectIcon(UObject& Object);
+
+	/** @return The text to use for displaying this property's name. Uses Class to determine class name if available. */
+	FText GetPropertyDisplayText(const FConcertPropertyChain& Property, UStruct* Class = nullptr);
+	/** @return More lightweight version of GetPropertyDisplayString which does not construct any FText. Uses Class to determine class name if available. */
+	FString GetPropertyDisplayString(const FConcertPropertyChain& Property, UStruct* Class = nullptr);
+}
diff --git a/Engine/Plugins/Developer/Concert/ConcertUI/ConcertSharedSlate/Source/ConcertSharedSlate/Private/Replication/Editor/Utils/ObjectUtils.cpp b/Engine/Plugins/Developer/Concert/ConcertUI/ConcertSharedSlate/Source/ConcertSharedSlate/Private/Replication/Editor/Utils/ObjectUtils.cpp
new file mode 100644
--- /dev/null
+++ b/Engine/Plugins/Developer/Concert/ConcertUI/ConcertSharedSlate/Source/ConcertSharedSlate/Private/Replication/Editor/Utils/ObjectUtils.cpp
@@ -0,0 +1,71 @@
+// Copyright Epic Games, Inc. All Rights Reserved.
+
+#include "Replication/Utils/ObjectUtils.h"
+
+#include "Containers/UnrealString.h"
+#include "UObject/SoftObjectPath.h"
+
+namespace UE::ConcertSharedSlate::ObjectUtils
+{
+	bool IsActor(const FSoftObjectPath& SoftObjectPath)
+	{
+		// Example of an actor called floor
+		// SoftObjectPath = { AssetPath = {PackageName = "/Game/Maps/SyncBoxLevel", AssetName = "SyncBoxLevel"}, SubPathString = "PersistentLevel.Floor" } }
+		const FString& SubPathString = SoftObjectPath.GetSubPathString();
+
+		constexpr int32 PersistentLevelStringLength = 16; // "PersistentLevel." has 16 characters
+		const bool bIsWorldObject = SubPathString.Contains(TEXT("PersistentLevel."), ESearchCase::CaseSensitive);
+		if (!bIsWorldObject)
+		{
+			// Not a path to a world object
+			return {};
+		}
+
+		// Start search after the . behind PersistentLevel
+		const int32 StartSearch = PersistentLevelStringLength + 1;
+		const int32 IndexOfDotAfterActorName = SubPathString.Find(TEXT("."), ESearchCase::CaseSensitive, ESearchDir::FromStart, StartSearch);
+		return IndexOfDotAfterActorName == INDEX_NONE;
+	}
+	
+	TOptional<FSoftObjectPath> GetActorOf(const FSoftObjectPath& SoftObjectPath)
+	{
+		// Example of an actor called floor
+		// SoftObjectPath = { AssetPath = {PackageName = "/Game/Maps/SyncBoxLevel", AssetName = "SyncBoxLevel"}, SubPathString = "PersistentLevel.Floor" } }
+		const FString& SubPathString = SoftObjectPath.GetSubPathString();
+
+		constexpr int32 PersistentLevelStringLength = 16; // "PersistentLevel." has 16 characters
+		const bool bIsWorldObject = SubPathString.Contains(TEXT("PersistentLevel."), ESearchCase::CaseSensitive);
+		if (!bIsWorldObject)
+		{
+			// Not a path to a world object
+			return {};
+		}
+
+		// Start search after the . behind PersistentLevel
+		const int32 StartSearch = PersistentLevelStringLength + 1;
+		const int32 IndexOfDotAfterActorName = SubPathString.Find(TEXT("."), ESearchCase::CaseSensitive, ESearchDir::FromStart, StartSearch);
+		if (IndexOfDotAfterActorName == INDEX_NONE)
+		{
+			// SoftObjectPath points to an actor
+			return {};
+		}
+
+		const int32 NumToChopOffRight = SubPathString.Len() - IndexOfDotAfterActorName;
+		const FString NewSubstring = SubPathString.LeftChop(NumToChopOffRight);
+		const FSoftObjectPath PathToOwningActor(SoftObjectPath.GetAssetPath(), NewSubstring);
+		return PathToOwningActor;
+	}
+	
+	FString ExtractObjectDisplayStringFromPath(const FSoftObjectPath& Object)
+	{
+		// Subpath looks like this PersistentLevel.Actor.Component
+		const FString& Subpath = Object.GetSubPathString();
+		const int32 LastDotIndex = Subpath.Find(TEXT("."), ESearchCase::CaseSensitive, ESearchDir::FromEnd);
+		if (LastDotIndex == INDEX_NONE)
+		{
+			return {};
+		}
+		return Subpath.RightChop(LastDotIndex + 1);
+	}
+};
+
diff --git a/Engine/Plugins/Developer/Concert/ConcertUI/ConcertSharedSlate/Source/ConcertSharedSlate/Private/Replication/Editor/Utils/PropertyEnumerationUtils.cpp b/Engine/Plugins/Developer/Concert/ConcertUI/ConcertSharedSlate/Source/ConcertSharedSlate/Private/Replication/Editor/Utils/PropertyEnumerationUtils.cpp
new file mode 100644
--- /dev/null
+++ b/Engine/Plugins/Developer/Concert/ConcertUI/ConcertSharedSlate/Source/ConcertSharedSlate/Private/Replication/Editor/Utils/PropertyEnumerationUtils.cpp
@@ -0,0 +1,92 @@
+// Copyright Epic Games, Inc. All Rights Reserved.
+
+#include "PropertyEnumerationUtils.h"
+
+#include "Misc/EBreakBehavior.h"
+#include "Replication/Editor/Model/IReplicationStreamModel.h"
+#include "Replication/Editor/Model/Property/IPropertySelectionSourceModel.h"
+
+#include "Containers/Set.h"
+#include "UObject/SoftObjectPath.h"
+
+namespace UE::ConcertSharedSlate
+{
+	/** Gets the class from the model or loads it. This function is designed to be used without assuming that it is run in editor-builds. */
+	FSoftClassPath GetObjectClassFromModelOrLoad(const FSoftObjectPath& Object, const IReplicationStreamModel& Model)
+	{
+		const FSoftClassPath ResolvedClass = Model.GetObjectClass(Object);
+		if (ResolvedClass.IsValid())
+		{
+			return ResolvedClass;
+		}
+
+#if WITH_EDITOR
+		// The object is not yet in the model.
+		const UObject* LoadedObject = Object.ResolveObject();
+		return LoadedObject ? LoadedObject->GetClass() : FSoftClassPath{};
+#else
+		return FSoftClassPath{};
+#endif
+	}
+
+	void EnumerateProperties(const TSet<FSoftObjectPath>& Objects, const IReplicationStreamModel& Model, const IPropertySelectionSourceModel* OptionalSource, FEnumerateProperties Callback)
+	{
+		if (OptionalSource)
+		{
+			EnumerateAllProperties(Objects, *OptionalSource, Model, Callback);
+		}
+		else
+		{
+			EnumerateRegisteredPropertiesOnly(Objects, Model, Callback);
+		}
+	}
+
+	/** Enumerates the properties that are assigned to the object in Model */
+	void EnumerateRegisteredPropertiesOnly(const TSet<FSoftObjectPath>& Objects, const IReplicationStreamModel& Model, FEnumerateProperties Callback)
+	{
+		for (const FSoftObjectPath& Object : Objects)
+		{
+			const FSoftClassPath ObjectClass = GetObjectClassFromModelOrLoad(Object, Model);
+
+			EBreakBehavior BreakBehavior = EBreakBehavior::Continue;
+			Model.ForEachProperty(Object, [&Callback, &ObjectClass, &BreakBehavior](const FConcertPropertyChain& Chain)
+			{
+				BreakBehavior = Callback(ObjectClass, Chain);
+				return BreakBehavior;
+			});
+
+			if (BreakBehavior == EBreakBehavior::Break)
+			{
+				break;
+			}
+		}
+	}
+
+	/** Enumerate the properties that are selectable in Source (e.g. all properties in that class, @see FSelectPropertyFromUClassModel). */
+	void EnumerateAllProperties(const TSet<FSoftObjectPath>& Objects, const IPropertySelectionSourceModel& Source, const IReplicationStreamModel& Model, FEnumerateProperties Callback)
+	{
+		TSet<FSoftClassPath> VisitedClasses; 
+			
+		for (const FSoftObjectPath& Object : Objects)
+		{
+			const FSoftClassPath ObjectClass = GetObjectClassFromModelOrLoad(Object, Model);
+			if (VisitedClasses.Contains(ObjectClass))
+			{
+				continue;
+			}
+			VisitedClasses.Add(ObjectClass);
+					
+			EBreakBehavior BreakBehavior = EBreakBehavior::Continue;
+			Source.GetPropertySource(ObjectClass)->EnumerateSelectableItems([&Callback, &ObjectClass, &BreakBehavior](const FSelectablePropertyInfo& PropertyInfo)
+			{
+				BreakBehavior = Callback(ObjectClass, PropertyInfo.Property);
+				return BreakBehavior;
+			});
+					
+			if (BreakBehavior == EBreakBehavior::Break)
+			{
+				break;
+			}
+		}
+	}
+}
\ No newline at end of file
diff --git a/Engine/Plugins/Developer/Concert/ConcertUI/ConcertSharedSlate/Source/ConcertSharedSlate/Private/Replication/Editor/Utils/PropertyEnumerationUtils.h b/Engine/Plugins/Developer/Concert/ConcertUI/ConcertSharedSlate/Source/ConcertSharedSlate/Private/Replication/Editor/Utils/PropertyEnumerationUtils.h
new file mode 100644
--- /dev/null
+++ b/Engine/Plugins/Developer/Concert/ConcertUI/ConcertSharedSlate/Source/ConcertSharedSlate/Private/Replication/Editor/Utils/PropertyEnumerationUtils.h
@@ -0,0 +1,30 @@
+// Copyright Epic Games, Inc. All Rights Reserved.
+
+#pragma once
+
+
+#include "Containers/ContainersFwd.h"
+#include "HAL/Platform.h"
+#include "Templates/FunctionFwd.h"
+
+enum class EBreakBehavior : uint8;
+struct FConcertPropertyChain;
+struct FSoftClassPath;
+struct FSoftObjectPath;
+
+namespace UE::ConcertSharedSlate
+{
+	class IPropertySelectionSourceModel;
+	class IReplicationStreamModel;
+	using FEnumerateProperties = TFunctionRef<EBreakBehavior(const FSoftClassPath&, const FConcertPropertyChain&)>;
+	
+	/** Gets the class from the model or loads it. This function is designed to be used without assuming that it is run in editor-builds. */
+	FSoftClassPath GetObjectClassFromModelOrLoad(const FSoftObjectPath& Object, const IReplicationStreamModel& Model);
+
+	/** Calls EnumerateRegisteredPropertiesOnly or EnumerateAllProperties depending on whether OptionalSource is nullptr. */
+	void EnumerateProperties(const TSet<FSoftObjectPath>& Objects, const IReplicationStreamModel& Model, const IPropertySelectionSourceModel* OptionalSource, FEnumerateProperties Callback);
+	/** Enumerates the properties that are assigned to the object in Model */
+	void EnumerateRegisteredPropertiesOnly(const TSet<FSoftObjectPath>& Objects, const IReplicationStreamModel& Model, FEnumerateProperties Callback);
+	/** Enumerate the properties that are selectable in Source (e.g. all properties in that class, @see FSelectPropertyFromUClassModel). */
+	void EnumerateAllProperties(const TSet<FSoftObjectPath>& Objects, const IPropertySelectionSourceModel& Source, const IReplicationStreamModel& Model, FEnumerateProperties Callback);
+}
\ No newline at end of file
diff --git a/Engine/Plugins/Developer/Concert/ConcertUI/ConcertSharedSlate/Source/ConcertSharedSlate/Private/Replication/Editor/View/Column/ObjectColumnAdapter.h b/Engine/Plugins/Developer/Concert/ConcertUI/ConcertSharedSlate/Source/ConcertSharedSlate/Private/Replication/Editor/View/Column/ObjectColumnAdapter.h
--- a/Engine/Plugins/Developer/Concert/ConcertUI/ConcertSharedSlate/Source/ConcertSharedSlate/Private/Replication/Editor/View/Column/ObjectColumnAdapter.h
+++ b/Engine/Plugins/Developer/Concert/ConcertUI/ConcertSharedSlate/Source/ConcertSharedSlate/Private/Replication/Editor/View/Column/ObjectColumnAdapter.h
@@ -2,13 +2,13 @@
 
 #pragma once
 
-#include "Replication/Editor/Model/ReplicatedObjectData.h"
+#include "Replication/Editor/Model/Data/ReplicatedObjectData.h"
 #include "Replication/Editor/View/Column/IObjectTreeColumn.h"
 #include "Replication/Editor/View/Column/IReplicationTreeColumn.h"
 
 namespace UE::ConcertSharedSlate
 {
-	/** Adapts an IObjectTreeColumn to IReplicationTreeColumn<FReplicatedPropertyData>. It simply passes additional info down to IPropertyTreeColumn. */
+	/** Adapts an IObjectTreeColumn to IReplicationTreeColumn<FPropertyData>. It simply passes additional info down to IPropertyTreeColumn. */
 	class FObjectColumnAdapter : public IReplicationTreeColumn<FReplicatedObjectData>
 	{
 	public:
diff --git a/Engine/Plugins/Developer/Concert/ConcertUI/ConcertSharedSlate/Source/ConcertSharedSlate/Private/Replication/Editor/View/Column/PropertyColumnAdapter.h b/Engine/Plugins/Developer/Concert/ConcertUI/ConcertSharedSlate/Source/ConcertSharedSlate/Private/Replication/Editor/View/Column/PropertyColumnAdapter.h
--- a/Engine/Plugins/Developer/Concert/ConcertUI/ConcertSharedSlate/Source/ConcertSharedSlate/Private/Replication/Editor/View/Column/PropertyColumnAdapter.h
+++ b/Engine/Plugins/Developer/Concert/ConcertUI/ConcertSharedSlate/Source/ConcertSharedSlate/Private/Replication/Editor/View/Column/PropertyColumnAdapter.h
@@ -2,24 +2,24 @@
 
 #pragma once
 
-#include "Replication/Editor/Model/ReplicatedPropertyData.h"
+#include "Replication/Editor/Model/Data/PropertyData.h"
 #include "Replication/Editor/View/Column/IPropertyTreeColumn.h"
 #include "Replication/Editor/View/Column/IReplicationTreeColumn.h"
 
 namespace UE::ConcertSharedSlate
 {
-	/** Adapts an IPropertyTreeColumn to IReplicationTreeColumn<FReplicatedPropertyData>. It simply passes additional info down to IPropertyTreeColumn. */
-	class FPropertyColumnAdapter : public IReplicationTreeColumn<FReplicatedPropertyData>
+	/** Adapts an IPropertyTreeColumn to IReplicationTreeColumn<FPropertyData>. It simply passes additional info down to IPropertyTreeColumn. */
+	class FPropertyColumnAdapter : public IReplicationTreeColumn<FPropertyData>
 	{
 	public:
 
-		static TArray<TReplicationColumnEntry<FReplicatedPropertyData>> Transform(const TArray<FPropertyColumnEntry>& Entries)
+		static TArray<TReplicationColumnEntry<FPropertyData>> Transform(const TArray<FPropertyColumnEntry>& Entries)
 		{
-			TArray<TReplicationColumnEntry<FReplicatedPropertyData>> Result;
-			Algo::Transform(Entries, Result, [](const FPropertyColumnEntry& Entry) -> TReplicationColumnEntry<FReplicatedPropertyData>
+			TArray<TReplicationColumnEntry<FPropertyData>> Result;
+			Algo::Transform(Entries, Result, [](const FPropertyColumnEntry& Entry) -> TReplicationColumnEntry<FPropertyData>
 			{
 				return {
-					TReplicationColumnDelegates<FReplicatedPropertyData>::FCreateColumn::CreateLambda([CreateDelegate = Entry.CreateColumn]()
+					TReplicationColumnDelegates<FPropertyData>::FCreateColumn::CreateLambda([CreateDelegate = Entry.CreateColumn]()
 					{
 						return MakeShared<FPropertyColumnAdapter>(CreateDelegate.Execute());
 					}),
@@ -39,19 +39,19 @@ namespace UE::ConcertSharedSlate
 		{
 			return AdaptedColumn->GenerateColumnWidget({ InArgs.HighlightText, Transform(InArgs.RowItem) });
 		}
-		virtual void PopulateSearchString(const FReplicatedPropertyData& InItem, TArray<FString>& InOutSearchStrings) const override
+		virtual void PopulateSearchString(const FPropertyData& InItem, TArray<FString>& InOutSearchStrings) const override
 		{
 			return AdaptedColumn->PopulateSearchString(Transform(InItem), InOutSearchStrings);
 		}
 
 		virtual bool CanBeSorted() const override { return AdaptedColumn->CanBeSorted(); }
-		virtual bool IsLessThan(const FReplicatedPropertyData& Left, const FReplicatedPropertyData& Right) const override { return AdaptedColumn->IsLessThan(Transform(Left), Transform(Right)); }
+		virtual bool IsLessThan(const FPropertyData& Left, const FPropertyData& Right) const override { return AdaptedColumn->IsLessThan(Transform(Left), Transform(Right)); }
 
 	private:
 
 		TSharedRef<IPropertyTreeColumn> AdaptedColumn;
 
-		static FPropertyTreeRowContext Transform(FReplicatedPropertyData Data)
+		static FPropertyTreeRowContext Transform(FPropertyData Data)
 		{
 			return { MoveTemp(Data) };
 		}
diff --git a/Engine/Plugins/Developer/Concert/ConcertUI/ConcertSharedSlate/Source/ConcertSharedSlate/Private/Replication/Editor/View/DisplayUtils.cpp b/Engine/Plugins/Developer/Concert/ConcertUI/ConcertSharedSlate/Source/ConcertSharedSlate/Private/Replication/Editor/View/DisplayUtils.cpp
deleted file mode 100644
diff --git a/Engine/Plugins/Developer/Concert/ConcertUI/ConcertSharedSlate/Source/ConcertSharedSlate/Private/Replication/Editor/View/DisplayUtils.h b/Engine/Plugins/Developer/Concert/ConcertUI/ConcertSharedSlate/Source/ConcertSharedSlate/Private/Replication/Editor/View/DisplayUtils.h
deleted file mode 100644
diff --git a/Engine/Plugins/Developer/Concert/ConcertUI/ConcertSharedSlate/Source/ConcertSharedSlate/Private/Replication/Editor/View/ObjectEditor/FakeObjectToPropertiesEditorModel.cpp b/Engine/Plugins/Developer/Concert/ConcertUI/ConcertSharedSlate/Source/ConcertSharedSlate/Private/Replication/Editor/View/ObjectEditor/FakeObjectToPropertiesEditorModel.cpp
deleted file mode 100644
diff --git a/Engine/Plugins/Developer/Concert/ConcertUI/ConcertSharedSlate/Source/ConcertSharedSlate/Private/Replication/Editor/View/ObjectEditor/FakeObjectToPropertiesEditorModel.h b/Engine/Plugins/Developer/Concert/ConcertUI/ConcertSharedSlate/Source/ConcertSharedSlate/Private/Replication/Editor/View/ObjectEditor/FakeObjectToPropertiesEditorModel.h
deleted file mode 100644
diff --git a/Engine/Plugins/Developer/Concert/ConcertUI/ConcertSharedSlate/Source/ConcertSharedSlate/Private/Replication/Editor/View/ObjectEditor/SBaseReplicationStreamEditor.cpp b/Engine/Plugins/Developer/Concert/ConcertUI/ConcertSharedSlate/Source/ConcertSharedSlate/Private/Replication/Editor/View/ObjectEditor/SBaseReplicationStreamEditor.cpp
--- a/Engine/Plugins/Developer/Concert/ConcertUI/ConcertSharedSlate/Source/ConcertSharedSlate/Private/Replication/Editor/View/ObjectEditor/SBaseReplicationStreamEditor.cpp
+++ b/Engine/Plugins/Developer/Concert/ConcertUI/ConcertSharedSlate/Source/ConcertSharedSlate/Private/Replication/Editor/View/ObjectEditor/SBaseReplicationStreamEditor.cpp
@@ -2,14 +2,13 @@
 
 #include "SBaseReplicationStreamEditor.h"
 
-#include "FakeObjectToPropertiesEditorModel.h"
 #include "Model/Item/SourceModelBuilders.h"
 #include "Replication/Editor/Model/IEditableReplicationStreamModel.h"
 #include "Replication/Editor/Model/Property/IPropertySelectionSourceModel.h"
-#include "Replication/Editor/Model/ReplicatedObjectData.h"
-#include "Replication/Editor/View/DisplayUtils.h"
-#include "Replication/Editor/View/ObjectViewer/SReplicationStreamViewer.h"
+#include "Replication/Editor/Model/Data/ReplicatedObjectData.h"
 #include "Replication/Editor/Model/ObjectSource/IObjectSelectionSourceModel.h"
+#include "Replication/Editor/Utils/DisplayUtils.h"
+#include "Replication/Editor/View/ObjectViewer/SReplicationStreamViewer.h"
 
 #include "Algo/AnyOf.h"
 #include "Replication/Editor/Model/Object/IObjectHierarchyModel.h"
@@ -33,8 +32,6 @@ namespace UE::ConcertSharedSlate
 		EditablePropertiesModel = MoveTemp(InPropertiesModel);
 		EditablePropertiesModel->OnObjectsChanged().AddSP(this, &SBaseReplicationStreamEditor::OnObjectsChanged);
 		EditablePropertiesModel->OnPropertiesChanged().AddSP(this, &SBaseReplicationStreamEditor::OnPropertiesChanged);
-		
-		PropertiesModelAdapter = MakeShared<FFakeObjectToPropertiesEditorModel>(EditablePropertiesModel.ToSharedRef(), PropertySelectionSource.ToSharedRef());
 
 		IsEditingEnabledAttribute = InArgs._IsEditingEnabled;
 		EditingDisabledToolTipTextAttribute = InArgs._EditingDisabledToolTipText;
@@ -42,8 +39,8 @@ namespace UE::ConcertSharedSlate
 		
 		ChildSlot
 		[
-			SAssignNew(ReplicationViewer, SReplicationStreamViewer, PropertiesModelAdapter.ToSharedRef())
-				.PropertyTreeView(InArgs._PropertyTreeView)
+			SAssignNew(ReplicationViewer, SReplicationStreamViewer, EditablePropertiesModel.ToSharedRef())
+				.PropertyAssignmentView(InArgs._PropertyAssignmentView)
 				.ObjectColumns(InArgs._ObjectColumns)
 				.PrimaryObjectSort(InArgs._PrimaryObjectSort)
 				.SecondaryObjectSort(InArgs._SecondaryObjectSort)
@@ -90,7 +87,7 @@ namespace UE::ConcertSharedSlate
 		ReplicationViewer->RequestPropertyColumnResort(ColumnId);
 	}
 
-	TArray<FSoftObjectPath> SBaseReplicationStreamEditor::GetObjectsBeingPropertyEdited() const
+	TSet<FSoftObjectPath> SBaseReplicationStreamEditor::GetObjectsBeingPropertyEdited() const
 	{
 		return ReplicationViewer->GetObjectsBeingPropertyEdited();
 	}
diff --git a/Engine/Plugins/Developer/Concert/ConcertUI/ConcertSharedSlate/Source/ConcertSharedSlate/Private/Replication/Editor/View/ObjectEditor/SBaseReplicationStreamEditor.h b/Engine/Plugins/Developer/Concert/ConcertUI/ConcertSharedSlate/Source/ConcertSharedSlate/Private/Replication/Editor/View/ObjectEditor/SBaseReplicationStreamEditor.h
--- a/Engine/Plugins/Developer/Concert/ConcertUI/ConcertSharedSlate/Source/ConcertSharedSlate/Private/Replication/Editor/View/ObjectEditor/SBaseReplicationStreamEditor.h
+++ b/Engine/Plugins/Developer/Concert/ConcertUI/ConcertSharedSlate/Source/ConcertSharedSlate/Private/Replication/Editor/View/ObjectEditor/SBaseReplicationStreamEditor.h
@@ -6,7 +6,7 @@
 #include "Replication/Editor/View/IReplicationStreamEditor.h"
 #include "Replication/Editor/View/Column/IObjectTreeColumn.h"
 #include "Replication/Editor/View/Column/ReplicationColumnsUtils.h"
-#include "Replication/ReplicationWidgetDelegates.h"
+#include "Replication/Utils/ReplicationWidgetDelegates.h"
 
 #include "Widgets/DeclarativeSyntaxSupport.h"
 
@@ -18,11 +18,11 @@ namespace UE::ConcertSharedSlate
 {
 	class FFakeObjectToPropertiesEditorModel;
 	class FReplicatedObjectData;
-	class FReplicatedPropertyData;
+	class FPropertyData;
 	class IEditableReplicationStreamModel;
 	class IObjectNameModel;
 	class IObjectSelectionSourceModel;
-	class IPropertyTreeView;
+	class IPropertyAssignmentView;
 	class IPropertySelectionSourceModel;
 	class IReplicationSubobjectView;
 	class IObjectHierarchyModel;
@@ -52,8 +52,8 @@ namespace UE::ConcertSharedSlate
 		
 		SLATE_BEGIN_ARGS(SBaseReplicationStreamEditor)
 		{}
-			/** Displays the properties in a tree view */
-			SLATE_ARGUMENT(TSharedPtr<IPropertyTreeView>, PropertyTreeView)
+			/** In the lower half of the editor, this view presents the properties associated with the object that is currently selected in the upper part of the view. */
+			SLATE_ARGUMENT(TSharedPtr<IPropertyAssignmentView>, PropertyAssignmentView)
 		
 			/** Additional columns to add to the object view */
 			SLATE_ARGUMENT(TArray<FObjectColumnEntry>, ObjectColumns)
@@ -92,7 +92,7 @@ namespace UE::ConcertSharedSlate
 		virtual void Refresh() override;
 		virtual void RequestObjectColumnResort(const FName& ColumnId) override;
 		virtual void RequestPropertyColumnResort(const FName& ColumnId) override;
-		virtual TArray<FSoftObjectPath> GetObjectsBeingPropertyEdited() const override;
+		virtual TSet<FSoftObjectPath> GetObjectsBeingPropertyEdited() const override;
 		//~ End IReplicationStreamEditor Interface
 
 	private:
@@ -102,11 +102,6 @@ namespace UE::ConcertSharedSlate
 
 		/** For reading and writting to the edited asset */
 		TSharedPtr<IEditableReplicationStreamModel> EditablePropertiesModel;
-		/**
-		 * Fakes to SObjectToPropertyView that all UClass properties are contained.
-		 * We inject checkboxes to SObjectToPropertyView which do the actual adding and removing.
-		 */
-		TSharedPtr<FFakeObjectToPropertiesEditorModel> PropertiesModelAdapter;
 		/** Can be null. If set, adds all subobjects to the top level view. */
 		TSharedPtr<IObjectHierarchyModel> ObjectHierarchy;
 		
diff --git a/Engine/Plugins/Developer/Concert/ConcertUI/ConcertSharedSlate/Source/ConcertSharedSlate/Private/Replication/Editor/View/ObjectViewer/Property/SPropertyTreeView.cpp b/Engine/Plugins/Developer/Concert/ConcertUI/ConcertSharedSlate/Source/ConcertSharedSlate/Private/Replication/Editor/View/ObjectViewer/Property/SPropertyTreeView.cpp
deleted file mode 100644
diff --git a/Engine/Plugins/Developer/Concert/ConcertUI/ConcertSharedSlate/Source/ConcertSharedSlate/Private/Replication/Editor/View/ObjectViewer/Property/SPropertyTreeView.h b/Engine/Plugins/Developer/Concert/ConcertUI/ConcertSharedSlate/Source/ConcertSharedSlate/Private/Replication/Editor/View/ObjectViewer/Property/SPropertyTreeView.h
deleted file mode 100644
diff --git a/Engine/Plugins/Developer/Concert/ConcertUI/ConcertSharedSlate/Source/ConcertSharedSlate/Private/Replication/Editor/View/ObjectViewer/SReplicatedPropertyView.cpp b/Engine/Plugins/Developer/Concert/ConcertUI/ConcertSharedSlate/Source/ConcertSharedSlate/Private/Replication/Editor/View/ObjectViewer/SReplicatedPropertyView.cpp
--- a/Engine/Plugins/Developer/Concert/ConcertUI/ConcertSharedSlate/Source/ConcertSharedSlate/Private/Replication/Editor/View/ObjectViewer/SReplicatedPropertyView.cpp
+++ b/Engine/Plugins/Developer/Concert/ConcertUI/ConcertSharedSlate/Source/ConcertSharedSlate/Private/Replication/Editor/View/ObjectViewer/SReplicatedPropertyView.cpp
@@ -3,7 +3,7 @@
 #include "SReplicatedPropertyView.h"
 
 #include "Replication/Editor/Model/IReplicationStreamModel.h"
-#include "Replication/Editor/Model/ReplicatedObjectData.h"
+#include "Replication/Editor/Model/Data/ReplicatedObjectData.h"
 
 #include "Algo/AllOf.h"
 #include "Widgets/Layout/SWidgetSwitcher.h"
@@ -13,13 +13,14 @@
 
 namespace UE::ConcertSharedSlate
 {
-	void SReplicatedPropertyView::Construct(const FArguments& InArgs, TSharedRef<IPropertyTreeView> InPropertyTreeView, TSharedRef<IReplicationStreamModel> InPropertiesModel)
+	void SReplicatedPropertyView::Construct(const FArguments& InArgs, TSharedRef<IPropertyAssignmentView> InPropertyAssignmentView, TSharedRef<IReplicationStreamModel> InPropertiesModel)
 	{
-		ReplicatedProperties = MoveTemp(InPropertyTreeView);
+		PropertyAssignmentView = MoveTemp(InPropertyAssignmentView);
 		PropertiesModel = MoveTemp(InPropertiesModel);
 		
 		GetSelectedRootObjectsDelegate = InArgs._GetSelectedRootObjects;
-		check(GetSelectedRootObjectsDelegate.IsBound());
+		GetObjectClassDelegate = InArgs._GetObjectClass;
+		check(GetSelectedRootObjectsDelegate.IsBound() && GetObjectClassDelegate.IsBound());
 		
 		ChildSlot
 		[
@@ -29,7 +30,7 @@ namespace UE::ConcertSharedSlate
 	
 	void SReplicatedPropertyView::RefreshPropertyData()
 	{
-		TArray<FSoftObjectPath> SelectedObjects = GetObjectsSelectedForPropertyEditing();
+		TSet<FSoftObjectPath> SelectedObjects = GetObjectsSelectedForPropertyEditing();
 		if (SelectedObjects.IsEmpty())
 		{
 			SetPropertyContent(EReplicatedPropertyContent::NoSelection);
@@ -44,35 +45,13 @@ namespace UE::ConcertSharedSlate
 			return;
 		}
 		
-		const FSoftClassPath Class = *SharedClass;
-		if (!ensure(Class.IsValid()))
-		{
-			return;
-		}
-		
-		// Build the set of properties that are shared by all of the selected objects
-		TSet<FConcertPropertyChain> SharedProperties = PropertiesModel->GetAllProperties(SelectedObjects[0]);
-		for (int32 i = 1; i < SelectedObjects.Num(); ++i)
-		{
-			SharedProperties = PropertiesModel->GetAllProperties(SelectedObjects[i]).Union(SharedProperties);
-		}
-		
-		// If the objects have changed, the classes may share properties.
-		// In that case, below we'd reuse the item pointer, which would cause the tree view to re-use the old row widgets.
-		// However, we must regenerate all column widgets since they may be referencing the object the row was originally built for. So they'd display the state of the previous object still!
-		// Example: Assign property combo-box in Multi-User All Clients view displays who has the property assigned.
-		// Note: If the objects did not change, we definitely want to reuse item pointers since otherwise the user row selection is reset.
-		const bool bCanReusePropertyData = PreviousSelectedObjects == SelectedObjects; // This SHOULD be an order independent compare but usually Num == 1, so whatever
-		
-		ReplicatedProperties->RefreshPropertyData(SharedProperties, Class, bCanReusePropertyData);
-		
+		PropertyAssignmentView->RefreshData(SelectedObjects, *PropertiesModel);
 		SetPropertyContent(EReplicatedPropertyContent::Properties);
-		PreviousSelectedObjects = MoveTemp(SelectedObjects);
 	}
 
-	TArray<FSoftObjectPath> SReplicatedPropertyView::GetObjectsSelectedForPropertyEditing() const
+	TSet<FSoftObjectPath> SReplicatedPropertyView::GetObjectsSelectedForPropertyEditing() const
 	{
-		TArray<FSoftObjectPath> Result;
+		TSet<FSoftObjectPath> Result;
 		Algo::Transform(GetSelectedRootObjectsDelegate.Execute(), Result, [](const TSharedPtr<FReplicatedObjectData>& ObjectData)
 		{
 			return ObjectData->GetObjectPath();
@@ -89,7 +68,7 @@ namespace UE::ConcertSharedSlate
 			// EReplicatedPropertyContent::Properties
 			+SWidgetSwitcher::Slot()
 			[
-				ReplicatedProperties->GetWidget()
+				PropertyAssignmentView->GetWidget()
 			]
 			
 			// EReplicatedPropertyContent::NoSelection
@@ -111,14 +90,16 @@ namespace UE::ConcertSharedSlate
 			];
 	}
 	
-	TOptional<FSoftClassPath> SReplicatedPropertyView::GetClassForPropertiesFromSelection(const TArray<FSoftObjectPath>& Objects) const
+	TOptional<FSoftClassPath> SReplicatedPropertyView::GetClassForPropertiesFromSelection(const TSet<FSoftObjectPath>& Objects) const
 	{
-		const FSoftClassPath Class = PropertiesModel->GetObjectClass(Objects[0]);
-		const bool bAllHaveSameClass = Algo::AllOf(Objects, [this, Class](const FSoftObjectPath& Object)
+		FSoftClassPath SharedClass;
+		const bool bAllHaveSameClass = Algo::AllOf(Objects, [this, &SharedClass](const FSoftObjectPath& Object)
 		{
-			return PropertiesModel->GetObjectClass(Object) == Class;
+			const FSoftClassPath ObjectClass = GetObjectClass(Object);
+			SharedClass = SharedClass.IsValid() ? SharedClass : ObjectClass;
+			return ObjectClass == SharedClass;
 		});
-		return bAllHaveSameClass ? Class : TOptional<FSoftClassPath>{};
+		return bAllHaveSameClass ? SharedClass : TOptional<FSoftClassPath>{};
 	}
 
 	void SReplicatedPropertyView::SetPropertyContent(EReplicatedPropertyContent Content) const
diff --git a/Engine/Plugins/Developer/Concert/ConcertUI/ConcertSharedSlate/Source/ConcertSharedSlate/Private/Replication/Editor/View/ObjectViewer/SReplicatedPropertyView.h b/Engine/Plugins/Developer/Concert/ConcertUI/ConcertSharedSlate/Source/ConcertSharedSlate/Private/Replication/Editor/View/ObjectViewer/SReplicatedPropertyView.h
--- a/Engine/Plugins/Developer/Concert/ConcertUI/ConcertSharedSlate/Source/ConcertSharedSlate/Private/Replication/Editor/View/ObjectViewer/SReplicatedPropertyView.h
+++ b/Engine/Plugins/Developer/Concert/ConcertUI/ConcertSharedSlate/Source/ConcertSharedSlate/Private/Replication/Editor/View/ObjectViewer/SReplicatedPropertyView.h
@@ -2,8 +2,9 @@
 
 #pragma once
 
-#include "Replication/Editor/View/IPropertyTreeView.h"
+#include "Replication/Utils/ReplicationWidgetDelegates.h"
 #include "Replication/Editor/View/Column/SelectionViewerColumns.h"
+#include "Replication/Editor/View/IPropertyAssignmentView.h"
 
 #include "Misc/Optional.h"
 #include "Widgets/DeclarativeSyntaxSupport.h"
@@ -27,23 +28,26 @@ namespace UE::ConcertSharedSlate
 		{}
 			/** Gets the root objects selected in the object outliner. */
 			SLATE_EVENT(FGetSelectedRootObjects, GetSelectedRootObjects)
+
+			/** Gets the class for the object since the object may not be in the model. */
+			SLATE_EVENT(FGetObjectClass, GetObjectClass)
 		
 			/** Optional. If set, this determines the display text for objects. */
 			SLATE_ARGUMENT(TSharedPtr<IObjectNameModel>, NameModel)
 		SLATE_END_ARGS()
 
-		void Construct(const FArguments& InArgs, TSharedRef<IPropertyTreeView> InPropertyTreeView, TSharedRef<IReplicationStreamModel> InPropertiesModel);
+		void Construct(const FArguments& InArgs, TSharedRef<IPropertyAssignmentView> InPropertyAssignmentView, TSharedRef<IReplicationStreamModel> InPropertiesModel);
 		
 		void RefreshPropertyData();
 		/** Requests that the given column be resorted, if it currently affects the row sorting. */
-		void RequestResortForColumn(const FName& ColumnId) const { ReplicatedProperties->RequestResortForColumn(ColumnId); }
+		void RequestResortForColumn(const FName& ColumnId) const { PropertyAssignmentView->RequestResortForColumn(ColumnId); }
 		
-		TArray<FSoftObjectPath> GetObjectsSelectedForPropertyEditing() const;
+		TSet<FSoftObjectPath> GetObjectsSelectedForPropertyEditing() const;
 
 	private:
 		
-		/** Tree view for replicated properties. Content depends on the current object selected. */
-		TSharedPtr<IPropertyTreeView> ReplicatedProperties;
+		/** In the lower half of the editor, this view presents the properties associated with the object that is currently selected in the upper part of the view. */
+		TSharedPtr<IPropertyAssignmentView> PropertyAssignmentView;
 		
 		/** The model this view is visualizing. */
 		TSharedPtr<IReplicationStreamModel> PropertiesModel;
@@ -62,16 +66,17 @@ namespace UE::ConcertSharedSlate
 		
 		/** Gets the root objects selected in the object outliner. */
 		FGetSelectedRootObjects GetSelectedRootObjectsDelegate;
-
-		/** Used to determine whether to rebuild the entire property data. */
-		TArray<FSoftObjectPath> PreviousSelectedObjects;
+		/** Gets the class for the object since the object may not be in the model. */
+		FGetObjectClass GetObjectClassDelegate;
 
 		TSharedRef<SWidget> CreatePropertiesView(const FArguments& InArgs);
 		
 		/** Given the selected objects, determines whether they all have the same class and returns it if so. */
-		TOptional<FSoftClassPath> GetClassForPropertiesFromSelection(const TArray<FSoftObjectPath>& Objects) const;
+		TOptional<FSoftClassPath> GetClassForPropertiesFromSelection(const TSet<FSoftObjectPath>& Objects) const;
 		/** Sets how to display this widget */
 		void SetPropertyContent(EReplicatedPropertyContent Content) const;
+
+		FSoftClassPath GetObjectClass(const FSoftObjectPath& ObjectPath) const { return GetObjectClassDelegate.Execute(ObjectPath); };
 	};
 }
 
diff --git a/Engine/Plugins/Developer/Concert/ConcertUI/ConcertSharedSlate/Source/ConcertSharedSlate/Private/Replication/Editor/View/ObjectViewer/SReplicationStreamViewer.cpp b/Engine/Plugins/Developer/Concert/ConcertUI/ConcertSharedSlate/Source/ConcertSharedSlate/Private/Replication/Editor/View/ObjectViewer/SReplicationStreamViewer.cpp
--- a/Engine/Plugins/Developer/Concert/ConcertUI/ConcertSharedSlate/Source/ConcertSharedSlate/Private/Replication/Editor/View/ObjectViewer/SReplicationStreamViewer.cpp
+++ b/Engine/Plugins/Developer/Concert/ConcertUI/ConcertSharedSlate/Source/ConcertSharedSlate/Private/Replication/Editor/View/ObjectViewer/SReplicationStreamViewer.cpp
@@ -5,17 +5,21 @@
 #include "ConcertFrontendUtils.h"
 #include "Replication/Editor/Model/IReplicationStreamModel.h"
 #include "Replication/Editor/Model/Object/IObjectHierarchyModel.h"
-#include "Replication/Editor/Model/ReplicatedObjectData.h"
+#include "Replication/Editor/Model/Data/ReplicatedObjectData.h"
 #include "Replication/Editor/View/Column/ObjectColumnAdapter.h"
 #include "Replication/Editor/View/Column/SelectionViewerColumns.h"
-#include "Replication/Editor/View/ObjectViewer/Property/SPropertyTreeView.h"
-#include "Replication/ObjectUtils.h"
+#include "Replication/Editor/View/Property/SPropertyTreeView.h"
+#include "Replication/Utils/ObjectUtils.h"
 #include "SReplicatedPropertyView.h"
 
 #include "Widgets/Layout/SBorder.h"
 #include "Widgets/Layout/SSplitter.h"
 #include "Widgets/SBoxPanel.h"
 #include "Widgets/Text/STextBlock.h"
+#if WITH_EDITOR
+#include "UObject/Class.h"
+#include "UObject/Object.h"
+#endif
 
 #define LOCTEXT_NAMESPACE "SObjectToPropertyView"
 
@@ -52,7 +56,7 @@ namespace UE::ConcertSharedSlate
 		PropertySection->RequestResortForColumn(ColumnId);
 	}
 
-	TArray<FSoftObjectPath> SReplicationStreamViewer::GetObjectsBeingPropertyEdited() const
+	TSet<FSoftObjectPath> SReplicationStreamViewer::GetObjectsBeingPropertyEdited() const
 	{
 		return PropertySection->GetObjectsSelectedForPropertyEditing();
 	}
@@ -119,17 +123,29 @@ namespace UE::ConcertSharedSlate
 			ReplicatedObjects->SetExpandedItems(ItemsToExpand, true);
 		}
 	}
-
+	
 	TArray<TSharedPtr<FReplicatedObjectData>> SReplicationStreamViewer::GetSelectedOutlinerObjects() const
 	{
 		TArray<TSharedPtr<FReplicatedObjectData>> SelectedItems = ReplicatedObjects->GetSelectedItems();
 		// Items may have been removed this tick. However, selected items may not have been updated yet because STreeView processes item changes at the end of tick. 
 		SelectedItems.SetNum(Algo::RemoveIf(SelectedItems, [this](const TSharedPtr<FReplicatedObjectData>& ObjectData)
 		{
-			const bool bIsInModel = PropertiesModel->ContainsObjects({ ObjectData->GetObjectPath() });
-			const TOptional<FSoftObjectPath> OwningActor = ObjectUtils::GetActorOf(ObjectData->GetObjectPath());
-			const bool bHierarchyIsInModel = OwningActor.IsSet() && PropertiesModel->ContainsObjects({ *OwningActor });
-			return !bIsInModel && !bHierarchyIsInModel;
+			const FSoftObjectPath& ObjectPath = ObjectData->GetObjectPath();
+			const bool bIsInModel = PropertiesModel->ContainsObjects({ ObjectPath });
+			
+			const TOptional<FSoftObjectPath> OwningActor = ObjectUtils::GetActorOf(ObjectPath);
+			// When displaying object from local machine...
+			// ... the "Add Actor" button has added an actor without properties to the model; however objects without assigned properties are not transmitted to server.
+			const bool bContainsOwningActor = OwningActor.IsSet() && PropertiesModel->ContainsObjects({ *OwningActor });
+
+			// When displaying object from remote machine...
+			// ... we only see objects with actual properties assigned.
+			const bool bIsSubobjectOfActor = OwningActor.IsSet() && ObjectPath.ToString().Contains(OwningActor->ToString());
+			// ... if ObjectData is an actor, we must also consider whether any of its subobjects is contained in the model
+			bool bContainsAnySubobject = false;
+			PropertiesModel->ForEachSubobject(ObjectPath, [&bContainsAnySubobject](const FSoftObjectPath& Child){ bContainsAnySubobject = true; return EBreakBehavior::Break; });
+			
+			return !bIsInModel && !bContainsOwningActor && !bIsSubobjectOfActor && !bContainsAnySubobject;
 		}));
 		return SelectedItems;
 	}
@@ -190,9 +206,11 @@ namespace UE::ConcertSharedSlate
 
 	TSharedRef<SWidget> SReplicationStreamViewer::CreateOutlinerSection(const FArguments& InArgs)
 	{
+		const FGetObjectClass GetObjectClassDelegate = FGetObjectClass::CreateSP(this, &SReplicationStreamViewer::GetObjectClass);
+		
 		TArray<FObjectColumnEntry> Columns = InArgs._ObjectColumns;
-		Columns.Add(ReplicationColumns::TopLevel::LabelColumn(PropertiesModel.ToSharedRef(), NameModel.Get()));
-		Columns.Add(ReplicationColumns::TopLevel::TypeColumn(PropertiesModel.ToSharedRef()));
+		Columns.Add(ReplicationColumns::TopLevel::LabelColumn(NameModel.Get(), GetObjectClassDelegate));
+		Columns.Add(ReplicationColumns::TopLevel::TypeColumn(GetObjectClassDelegate));
 		
 		const bool bHasNoOutlinerObjectsAttribute = InArgs._NoOutlinerObjects.IsBound() || InArgs._NoOutlinerObjects.IsSet(); 
 		const TAttribute<FText> NoObjectsAttribute = bHasNoOutlinerObjectsAttribute ? InArgs._NoOutlinerObjects : LOCTEXT("NoObjects", "No objects to display");
@@ -265,8 +283,9 @@ namespace UE::ConcertSharedSlate
 				]
 				.BodyContent()
 				[
-					SAssignNew(PropertySection, SReplicatedPropertyView, InArgs._PropertyTreeView.ToSharedRef(), PropertiesModel.ToSharedRef())
-					.GetSelectedRootObjects_Lambda([this](){ return GetSelectedOutlinerObjects(); })
+					SAssignNew(PropertySection, SReplicatedPropertyView, InArgs._PropertyAssignmentView.ToSharedRef(), PropertiesModel.ToSharedRef())
+					.GetSelectedRootObjects(this, &SReplicationStreamViewer::GetSelectedOutlinerObjects)
+					.GetObjectClass(this, &SReplicationStreamViewer::GetObjectClass)
 					.NameModel(InArgs._NameModel)
 				]
 			];
@@ -284,7 +303,7 @@ namespace UE::ConcertSharedSlate
 		// Do a complete refresh.
 		// Complete refresh is acceptable because the list is updated infrequently and typically small < 500 items.
 		// An alternative would be to change RefreshObjectData to be called with two variables ObjectsAdded and ObjectsRemoved.
-		PropertiesModel->ForEachReplicatedObject([this, &NewPathToObjectDataCache](const FSoftObjectPath& ObjectPath) mutable
+		IterateDisplayableObjects([this, &NewPathToObjectDataCache](const FSoftObjectPath& ObjectPath) mutable
 		{
 			TOptional<IObjectHierarchyModel::FParentInfo> ParentInfo = ObjectHierarchy->GetParentInfo(ObjectPath);
 			const bool bIsActor = !ParentInfo; 
@@ -314,6 +333,57 @@ namespace UE::ConcertSharedSlate
 		}
 	}
 
+	void SReplicationStreamViewer::IterateDisplayableObjects(TFunctionRef<void(const FSoftObjectPath& Object)> Delegate) const
+	{
+		// Case: RealModel contains only components but not the owning actor.
+		// In that case, we want the UI to still show the owning actor.
+		// We'll track this with these containers:
+		TSet<FSoftObjectPath> NonActors;
+		TSet<FSoftObjectPath> Actors;
+
+		// Here we only exposes objects that are top level so we skip components, etc.
+		EBreakBehavior BreakBehavior = EBreakBehavior::Continue;
+		PropertiesModel->ForEachReplicatedObject([this, &Delegate, &NonActors, &Actors, &BreakBehavior](const FSoftObjectPath& Object)
+		{
+			Delegate(Object);
+			
+			if (ObjectUtils::IsActor(Object))
+			{
+				Actors.Add(Object);
+			}
+			else
+			{
+				NonActors.Add(Object);
+			}
+			
+			return EBreakBehavior::Continue;
+		});
+
+		// No additional objects to show
+		if (BreakBehavior == EBreakBehavior::Break)
+		{
+			return;
+		}
+
+		// Now determine the top level objects that are not in RealModel but that need to be shown for subobjects that are in RealModel
+		for (const FSoftObjectPath& NonTopLevelObject : NonActors)
+		{
+			const TOptional<FSoftObjectPath> TopLevelObject = ObjectUtils::GetActorOf(NonTopLevelObject);
+			if (!TopLevelObject)
+			{
+				continue;
+			}
+			
+			if (!Actors.Contains(*TopLevelObject))
+			{
+				Delegate(*TopLevelObject);
+				
+				// Avoid calling Delegate again for TopLevelObject
+				Actors.Add(*TopLevelObject);
+			}
+		}
+	}
+
 	void SReplicationStreamViewer::RefreshPropertyData()
 	{
 		PropertySection->RefreshPropertyData();
@@ -408,6 +478,25 @@ namespace UE::ConcertSharedSlate
 		const bool bSkipSubobject = Relationship == EChildRelationship::Subobject && !ObjectViewOptions.ShouldDisplaySubobjects();
 		return !bSkipSubobject;
 	}
+
+	FSoftClassPath SReplicationStreamViewer::GetObjectClass(const FSoftObjectPath& ObjectPath) const
+	{
+		const FSoftClassPath ResolvedClass = PropertiesModel->GetObjectClass(ObjectPath);
+		if (ResolvedClass.IsValid())
+		{
+			return ResolvedClass;
+		}
+
+#if WITH_EDITOR
+		// In the editor, we display the entire hierarchy (see BuildObjectHierarchyIfNeeded) so some items may not be in PropertiesModel.
+		// Example: Add an actor with many components and assign nothing - all of those components will take this path.
+		const UObject* LoadedObject = ObjectPath.ResolveObject();
+		return LoadedObject ? LoadedObject->GetClass() : FSoftClassPath{};
+#else
+		// For non-editor, we should probably consider getting the class information through a delegate.
+		return FSoftClassPath{};
+#endif
+	}
 }
 
 #undef LOCTEXT_NAMESPACE
\ No newline at end of file
diff --git a/Engine/Plugins/Developer/Concert/ConcertUI/ConcertSharedSlate/Source/ConcertSharedSlate/Private/Replication/Editor/View/ObjectViewer/SReplicationStreamViewer.h b/Engine/Plugins/Developer/Concert/ConcertUI/ConcertSharedSlate/Source/ConcertSharedSlate/Private/Replication/Editor/View/ObjectViewer/SReplicationStreamViewer.h
--- a/Engine/Plugins/Developer/Concert/ConcertUI/ConcertSharedSlate/Source/ConcertSharedSlate/Private/Replication/Editor/View/ObjectViewer/SReplicationStreamViewer.h
+++ b/Engine/Plugins/Developer/Concert/ConcertUI/ConcertSharedSlate/Source/ConcertSharedSlate/Private/Replication/Editor/View/ObjectViewer/SReplicationStreamViewer.h
@@ -4,7 +4,7 @@
 
 #include "Replication/Editor/View/IReplicationStreamViewer.h"
 
-#include "Replication/Editor/Model/ReplicatedObjectData.h"
+#include "Replication/Editor/Model/Data/ReplicatedObjectData.h"
 #include "Replication/Editor/View/Tree/SReplicationTreeView.h"
 #include "Replication/Editor/View/Column/IObjectTreeColumn.h"
 #include "Replication/Editor/View/Column/SelectionViewerColumns.h"
@@ -24,12 +24,13 @@ struct FSoftObjectPath;
 namespace UE::ConcertSharedSlate
 {
 	class SReplicatedPropertyView;
-	class FReplicatedPropertyData;
+	class FPropertyData;
 	class FReplicatedObjectData;
 	class IEditableReplicationStreamModel;
+	class IObjectHierarchyModel;
 	class IObjectNameModel;
+	class IPropertyAssignmentView;
 	class IReplicationStreamModel;
-	class IObjectHierarchyModel;
 	class SPropertyTreeView;
 	
 	enum class EChildRelationship : uint8;
@@ -50,8 +51,8 @@ namespace UE::ConcertSharedSlate
 
 		SLATE_BEGIN_ARGS(SReplicationStreamViewer)
 		{}
-			/** Displays the properties in a tree view */
-			SLATE_ARGUMENT(TSharedPtr<IPropertyTreeView>, PropertyTreeView)
+			/** In the lower half of the editor, this view presents the properties associated with the object that is currently selected in the upper part of the view. */
+			SLATE_ARGUMENT(TSharedPtr<IPropertyAssignmentView>, PropertyAssignmentView)
 		
 			/** Additional columns to add to the object view */
 			SLATE_ARGUMENT(TArray<FObjectColumnEntry>, ObjectColumns)
@@ -87,7 +88,7 @@ namespace UE::ConcertSharedSlate
 		virtual void Refresh() override;
 		virtual void RequestObjectColumnResort(const FName& ColumnId) override;
 		virtual void RequestPropertyColumnResort(const FName& ColumnId) override;
-		virtual TArray<FSoftObjectPath> GetObjectsBeingPropertyEdited() const override;
+		virtual TSet<FSoftObjectPath> GetObjectsBeingPropertyEdited() const override;
 		//~ End IReplicationStreamViewer Interface
 
 		void RequestObjectDataRefresh() { bHasRequestedObjectRefresh = true; }
@@ -145,8 +146,13 @@ namespace UE::ConcertSharedSlate
 		TSharedRef<SWidget> CreateContentWidget(const FArguments& InArgs);
 		TSharedRef<SWidget> CreateOutlinerSection(const FArguments& InArgs);
 		TSharedRef<SWidget> CreatePropertiesSection(const FArguments& InArgs);
-		
+
+		/** Regenerates object row data in the top view */
 		void RefreshObjectData();
+		/** Lists all objects that need an object in the top-view. */
+		void IterateDisplayableObjects(TFunctionRef<void(const FSoftObjectPath& Object)> Delegate) const;
+		
+		/** Regenerates property row data in the bottom videw */
 		void RefreshPropertyData();
 		
 		/** Sets RootObjectRowData to all non-root nodes from ObjectRowData. */
@@ -167,5 +173,7 @@ namespace UE::ConcertSharedSlate
 			RequestPropertyDataRefresh();
 		}
 		bool ShouldDisplayObject(const FSoftObjectPath& Object, EChildRelationship Relationship) const;
+		
+		FSoftClassPath GetObjectClass(const FSoftObjectPath& ObjectPath) const;
 	};
 }
diff --git a/Engine/Plugins/Developer/Concert/ConcertUI/ConcertSharedSlate/Source/ConcertSharedSlate/Private/Replication/Editor/View/Property/SPerObjectPropertyAssignment.cpp b/Engine/Plugins/Developer/Concert/ConcertUI/ConcertSharedSlate/Source/ConcertSharedSlate/Private/Replication/Editor/View/Property/SPerObjectPropertyAssignment.cpp
new file mode 100644
--- /dev/null
+++ b/Engine/Plugins/Developer/Concert/ConcertUI/ConcertSharedSlate/Source/ConcertSharedSlate/Private/Replication/Editor/View/Property/SPerObjectPropertyAssignment.cpp
@@ -0,0 +1,65 @@
+// Copyright Epic Games, Inc. All Rights Reserved.
+
+#include "SPerObjectPropertyAssignment.h"
+
+#include "Misc/EBreakBehavior.h"
+#include "Replication/Data/ConcertPropertySelection.h"
+#include "Replication/Editor/Utils/PropertyEnumerationUtils.h"
+
+namespace UE::ConcertSharedSlate
+{
+	void SPerObjectPropertyAssignment::Construct(const FArguments& InArgs, TSharedRef<IPropertyTreeView> InTreeView)
+	{
+		TreeView = MoveTemp(InTreeView);
+		OptionalPropertySource = InArgs._PropertySource;
+		
+		ChildSlot
+		[
+			TreeView->GetWidget()
+		];
+	}
+
+	void SPerObjectPropertyAssignment::RefreshData(const TSet<FSoftObjectPath>& Objects, const IReplicationStreamModel& Model)
+	{
+		TSet<FConcertPropertyChain> Properties;
+		bool bHasClassPath = false;
+		FSoftClassPath ClassPath;
+
+		EnumerateProperties(Objects, Model, OptionalPropertySource.Get(),
+			[&Properties, &bHasClassPath, &ClassPath](const FSoftClassPath& Class, const FConcertPropertyChain& Chain)
+			{
+				if (!bHasClassPath)
+				{
+					bHasClassPath = true;
+					ClassPath = Class;
+				}
+				else if (ClassPath != Class)
+				{
+					// Duplicate class - caller of RefreshData should not have called us but handle it gracefully.
+					bHasClassPath = false;
+					return EBreakBehavior::Break;
+				}
+				
+				Properties.Add(Chain);
+				return EBreakBehavior::Continue;
+			});
+
+		if (bHasClassPath)
+		{
+			// If the objects have changed, the classes may share properties.
+			// In that case, below we'd reuse the item pointer, which would cause the tree view to re-use the old row widgets.
+			// However, we must regenerate all column widgets since they may be referencing the object the row was originally built for. So they'd display the state of the previous object still!
+			// Example: Assign property combo-box in Multi-User All Clients view displays who has the property assigned.
+			// Note: If the objects did not change, we definitely want to reuse item pointers since otherwise the user row selection is reset.
+			const bool bCanReusePropertyData = PreviousSelectedObjects.Num() == Objects.Num() && PreviousSelectedObjects.Includes(Objects);
+			PreviousSelectedObjects = Objects;
+			
+			TreeView->RefreshPropertyData(Properties, ClassPath, bCanReusePropertyData);
+		}
+		else
+		{
+			PreviousSelectedObjects.Reset();
+			TreeView->RefreshPropertyData(Properties, ClassPath, false);
+		}
+	}
+}
diff --git a/Engine/Plugins/Developer/Concert/ConcertUI/ConcertSharedSlate/Source/ConcertSharedSlate/Private/Replication/Editor/View/Property/SPerObjectPropertyAssignment.h b/Engine/Plugins/Developer/Concert/ConcertUI/ConcertSharedSlate/Source/ConcertSharedSlate/Private/Replication/Editor/View/Property/SPerObjectPropertyAssignment.h
new file mode 100644
--- /dev/null
+++ b/Engine/Plugins/Developer/Concert/ConcertUI/ConcertSharedSlate/Source/ConcertSharedSlate/Private/Replication/Editor/View/Property/SPerObjectPropertyAssignment.h
@@ -0,0 +1,49 @@
+// Copyright Epic Games, Inc. All Rights Reserved.
+
+#pragma once
+
+#include "Replication/Editor/View/IPropertyAssignmentView.h"
+#include "Replication/Editor/View/IPropertyTreeView.h"
+
+#include "Widgets/DeclarativeSyntaxSupport.h"
+#include "Widgets/SCompoundWidget.h"
+
+namespace UE::ConcertSharedSlate
+{
+	class IPropertySelectionSourceModel;
+	class IReplicationStreamModel;
+	
+	/** SPerObjectPropertyAssignment shows only the properties of the displayed object, which is achieved by wrapping SPropertyTreeView. */
+	class SPerObjectPropertyAssignment
+		: public SCompoundWidget
+		, public IPropertyAssignmentView
+	{
+	public:
+
+		SLATE_BEGIN_ARGS(SPerObjectPropertyAssignment){}
+			/** Optional. If specified, displays the properties of this model instead of those assigned in the stream. */
+			SLATE_ARGUMENT(TSharedPtr<IPropertySelectionSourceModel>, PropertySource)
+		SLATE_END_ARGS()
+
+		void Construct(const FArguments& InArgs, TSharedRef<IPropertyTreeView> InTreeView);
+
+		//~ Begin IPropertyAssignmentView Interface
+		virtual void RefreshData(const TSet<FSoftObjectPath>& Object, const IReplicationStreamModel& Model) override;
+		virtual void RequestRefilter() const override { return TreeView->RequestRefilter(); }
+		virtual void RequestResortForColumn(const FName& ColumnId) override { return TreeView->RequestRefilter(); }
+		virtual TSharedRef<SWidget> GetWidget() override { return SharedThis(this); }
+		//~ End IPropertyAssignmentView Interface
+
+	private:
+
+		/** The tree view that is being wrapped. */
+		TSharedPtr<IPropertyTreeView> TreeView;
+		/** Used to determine whether to rebuild the entire property data. */
+		TSet<FSoftObjectPath> PreviousSelectedObjects;
+		
+		/** Optional. If specified, displays the properties of this model instead of those assigned in the stream. */
+		TSharedPtr<IPropertySelectionSourceModel> OptionalPropertySource;
+	};
+}
+
+
diff --git a/Engine/Plugins/Developer/Concert/ConcertUI/ConcertSharedSlate/Source/ConcertSharedSlate/Private/Replication/Editor/View/Property/SPropertyTreeView.cpp b/Engine/Plugins/Developer/Concert/ConcertUI/ConcertSharedSlate/Source/ConcertSharedSlate/Private/Replication/Editor/View/Property/SPropertyTreeView.cpp
new file mode 100644
--- /dev/null
+++ b/Engine/Plugins/Developer/Concert/ConcertUI/ConcertSharedSlate/Source/ConcertSharedSlate/Private/Replication/Editor/View/Property/SPropertyTreeView.cpp
@@ -0,0 +1,110 @@
+// Copyright Epic Games, Inc. All Rights Reserved.
+
+#include "SPropertyTreeView.h"
+
+#include "Replication/Editor/Model/Data/PropertyData.h"
+
+#include "Algo/ForEach.h"
+#include "Replication/Editor/View/Column/PropertyColumnAdapter.h"
+#include "UObject/UnrealType.h"
+
+#define LOCTEXT_NAMESPACE "SReplicatedPropertiesView"
+
+namespace UE::ConcertSharedSlate
+{
+	void SPropertyTreeView::Construct(const FArguments& InArgs)
+	{
+		ChildSlot
+		[
+			SAssignNew(TreeView, SReplicationTreeView<FPropertyData>)
+				.RootItemsSource(&RootPropertyRowData)
+				.OnGetChildren(this, &SPropertyTreeView::GetPropertyRowChildren)
+				.FilterItem(InArgs._FilterItem)
+				.Columns(FPropertyColumnAdapter::Transform(InArgs._Columns))
+				.ExpandableColumnLabel(InArgs._ExpandableColumnLabel)
+				.PrimarySort(InArgs._PrimarySort)
+				.SecondarySort(InArgs._SecondarySort)
+				.SelectionMode(InArgs._SelectionMode)
+				.LeftOfSearchBar() [ InArgs._LeftOfSearchBar.Widget ]
+				.RightOfSearchBar() [ InArgs._RightOfSearchBar.Widget ]
+				.RowBelowSearchBar() [ InArgs._RowBelowSearchBar.Widget ]
+				.NoItemsContent() [ InArgs._NoItemsContent.Widget ]
+		];
+	}
+
+	void SPropertyTreeView::RefreshPropertyData(const TSet<FConcertPropertyChain>& PropertiesToDisplay, const FSoftClassPath& Class, const bool bCanReuseExistingRowItems)
+	{
+		if (!bCanReuseExistingRowItems)
+		{
+			ChainToPropertyDataCache.Reset();
+		}
+		
+		// Try to re-use old instances by using the old ChainToPropertyDataCache. This is also done so the expansion states restore correctly in the tree view.
+		TMap<FConcertPropertyChain, TSharedPtr<FPropertyData>> NewChainToPropertyDataCache;
+		
+		PropertyRowData.Empty();
+		for (const FConcertPropertyChain& PropertyChain : PropertiesToDisplay)
+		{
+			const TSharedPtr<FPropertyData>* ExistingItem = ChainToPropertyDataCache.Find(PropertyChain);
+			const TSharedRef<FPropertyData> Item = ExistingItem ? ExistingItem->ToSharedRef() : AllocatePropertyData(Class, PropertyChain);
+			PropertyRowData.Emplace(Item);
+			NewChainToPropertyDataCache.Emplace(PropertyChain, Item);
+		}
+
+		// If an item was removed, then NewPathToPropertyDataCache does not contain it. 
+		ChainToPropertyDataCache = MoveTemp(NewChainToPropertyDataCache);
+		
+		// The tree view requires the item source to only contain the root items.
+		BuildRootPropertyRowData();
+
+		TreeView->RequestRefilter();
+	}
+
+	void SPropertyTreeView::RequestScrollIntoView(const FConcertPropertyChain& PropertyChain)
+	{
+		const int32 Index = PropertyRowData.IndexOfByPredicate([&PropertyChain](const TSharedPtr<FPropertyData>& Data)
+		{
+			return Data->GetProperty() == PropertyChain;
+		});
+		if (PropertyRowData.IsValidIndex(Index))
+		{
+			TreeView->SetExpandedItems({ PropertyRowData[Index] }, true);
+			TreeView->RequestScrollIntoView(PropertyRowData[Index]);
+		}
+	}
+
+	TSharedRef<FPropertyData> SPropertyTreeView::AllocatePropertyData(FSoftClassPath OwningClass, FConcertPropertyChain PropertyChain)
+	{
+		return MakeShared<FPropertyData>(MoveTemp(OwningClass), MoveTemp(PropertyChain));
+	}
+
+	void SPropertyTreeView::BuildRootPropertyRowData()
+	{
+		RootPropertyRowData.Empty(PropertyRowData.Num());
+		for (const TSharedPtr<FPropertyData>& PropertyData : PropertyRowData)
+		{
+			if (PropertyData->GetProperty().IsRootProperty())
+			{
+				RootPropertyRowData.Emplace(PropertyData);
+			}
+		}
+	}
+
+	void SPropertyTreeView::GetPropertyRowChildren(TSharedPtr<FPropertyData> ReplicatedPropertyData, TFunctionRef<void(TSharedPtr<FPropertyData>)> ProcessChild)
+	{
+		TArray<TSharedPtr<FPropertyData>> Children;
+		
+		// Not the most efficient but it should be fine.
+		for (const TSharedPtr<FPropertyData>& Data : PropertyRowData)
+		{
+			if (Data->GetProperty().IsDirectChildOf(ReplicatedPropertyData->GetProperty()))
+			{
+				Children.Add(Data);
+			}
+		}
+
+		Algo::ForEach(Children, [&ProcessChild](const TSharedPtr<FPropertyData>& Data){ ProcessChild(Data); });
+	}
+}
+
+#undef LOCTEXT_NAMESPACE
\ No newline at end of file
diff --git a/Engine/Plugins/Developer/Concert/ConcertUI/ConcertSharedSlate/Source/ConcertSharedSlate/Private/Replication/Editor/View/Property/SPropertyTreeView.h b/Engine/Plugins/Developer/Concert/ConcertUI/ConcertSharedSlate/Source/ConcertSharedSlate/Private/Replication/Editor/View/Property/SPropertyTreeView.h
new file mode 100644
--- /dev/null
+++ b/Engine/Plugins/Developer/Concert/ConcertUI/ConcertSharedSlate/Source/ConcertSharedSlate/Private/Replication/Editor/View/Property/SPropertyTreeView.h
@@ -0,0 +1,88 @@
+// Copyright Epic Games, Inc. All Rights Reserved.
+
+#pragma once
+
+#include "Replication/Data/ConcertPropertySelection.h"
+#include "Replication/Editor/View/IPropertyTreeView.h"
+#include "Replication/Editor/View/Tree/SReplicationTreeView.h"
+
+#include "Replication/Editor/View/Column/IPropertyTreeColumn.h"
+#include "Widgets/DeclarativeSyntaxSupport.h"
+#include "Widgets/SCompoundWidget.h"
+
+namespace UE::ConcertSharedSlate
+{
+	class FPropertyData;
+	
+	/**
+	 * This widget knows how to display a list of properties in a tree view.
+	 * It generates the items and exposes extension points for more advanced UI, such as filtering.
+	 */
+	class SPropertyTreeView
+		: public SCompoundWidget
+		, public IPropertyTreeView
+	{
+	public:
+		
+		SLATE_BEGIN_ARGS(SPropertyTreeView)
+		{}
+			/*************** Arguments inherited by SReplicationTreeView ***************/
+		
+			/** Optional callback to do even more filtering of items. */
+			SLATE_EVENT(SReplicationTreeView<FPropertyData>::FCustomFilter, FilterItem)
+		
+			/** The columns this list should have */
+			SLATE_ARGUMENT(TArray<FPropertyColumnEntry>, Columns)
+			/** The name of the column that will have the SExpanderArrow for the tree view. */
+			SLATE_ARGUMENT(FName, ExpandableColumnLabel)
+			/** Initial primary sort to set. */
+			SLATE_ARGUMENT(FColumnSortInfo, PrimarySort)
+			/** Initial secondary sort to set. */
+			SLATE_ARGUMENT(FColumnSortInfo, SecondarySort)
+		
+			/** How many items are to allowed to be selected */
+			SLATE_ARGUMENT(ESelectionMode::Type, SelectionMode)
+		
+			/** Optional widget to add to the left of the search bar. */
+			SLATE_NAMED_SLOT(FArguments, LeftOfSearchBar)
+			/** Optional widget to add to the left of the search bar. */
+			SLATE_NAMED_SLOT(FArguments, RightOfSearchBar)
+		
+			/** Optional widget to add between the search bar and the table view (e.g. a SBasicFilterBar). */
+			SLATE_NAMED_SLOT(FArguments, RowBelowSearchBar)
+		
+			/** Optional, alternate content to show instead of the tree view when there are no rows. */
+			SLATE_NAMED_SLOT(FArguments, NoItemsContent)
+
+		SLATE_END_ARGS()
+
+		void Construct(const FArguments& InArgs);
+
+		//~ Begin IPropertyTreeView Interface
+		virtual void RefreshPropertyData(const TSet<FConcertPropertyChain>& PropertiesToDisplay, const FSoftClassPath& Class, bool bCanReuseExistingRowItems) override;
+		virtual void RequestRefilter() const override { TreeView->RequestRefilter(); }
+		virtual void RequestResortForColumn(const FName& ColumnId) override { TreeView->RequestResortForColumn(ColumnId); }
+		virtual void RequestScrollIntoView(const FConcertPropertyChain& PropertyChain) override;
+		virtual TSharedRef<SWidget> GetWidget() override { return SharedThis(this); }
+		//~ Begin IPropertyTreeView Interface
+
+	private:
+
+		/** The tree view displaying the replicated properties */
+		TSharedPtr<SReplicationTreeView<FPropertyData>> TreeView;
+		
+		/** Contains all data. */
+		TArray<TSharedPtr<FPropertyData>> PropertyRowData;
+		/** The instances which do not have any parents. This acts as the item source for the tree view. */
+		TArray<TSharedPtr<FPropertyData>> RootPropertyRowData;
+		/** Inverse map of PropertyRowData using FPropertyData::GetProperty as key. Contains all elements of PropertyRowData. */
+		TMap<FConcertPropertyChain, TSharedPtr<FPropertyData>> ChainToPropertyDataCache;
+		
+		TSharedRef<FPropertyData> AllocatePropertyData(FSoftClassPath OwningClass, FConcertPropertyChain PropertyChain);
+
+		/** Inits RootPropertyRowData from PropertyRowData. */
+		void BuildRootPropertyRowData();
+		void GetPropertyRowChildren(TSharedPtr<FPropertyData> ReplicatedPropertyData, TFunctionRef<void(TSharedPtr<FPropertyData>)> ProcessChild);
+	};
+}
+
diff --git a/Engine/Plugins/Developer/Concert/ConcertUI/ConcertSharedSlate/Source/ConcertSharedSlate/Private/Replication/Editor/View/SelectionViewerColumns.cpp b/Engine/Plugins/Developer/Concert/ConcertUI/ConcertSharedSlate/Source/ConcertSharedSlate/Private/Replication/Editor/View/SelectionViewerColumns.cpp
--- a/Engine/Plugins/Developer/Concert/ConcertUI/ConcertSharedSlate/Source/ConcertSharedSlate/Private/Replication/Editor/View/SelectionViewerColumns.cpp
+++ b/Engine/Plugins/Developer/Concert/ConcertUI/ConcertSharedSlate/Source/ConcertSharedSlate/Private/Replication/Editor/View/SelectionViewerColumns.cpp
@@ -4,9 +4,9 @@
 
 #include "ConcertFrontendStyle.h"
 #include "Replication/Editor/Model/IEditableReplicationStreamModel.h"
-#include "Replication/Editor/Model/ReplicatedPropertyData.h"
-#include "Replication/Editor/Model/ReplicatedObjectData.h"
-#include "Replication/Editor/View/DisplayUtils.h"
+#include "Replication/Editor/Model/Data/PropertyData.h"
+#include "Replication/Editor/Model/Data/ReplicatedObjectData.h"
+#include "Replication/Editor/Utils/DisplayUtils.h"
 #include "Replication/Editor/View/Column/ReplicationColumnsUtils.h"
 #include "Replication/PropertyChainUtils.h"
 
@@ -28,15 +28,15 @@ namespace UE::ConcertSharedSlate::ReplicationColumns::TopLevel
 	const FName LabelColumnId = TEXT("LabelColumn");
 	const FName TypeColumnId = TEXT("TypeColumn");
 	
-	FObjectColumnEntry LabelColumn(TSharedRef<IReplicationStreamModel> Model, IObjectNameModel* OptionalNameModel)
+	FObjectColumnEntry LabelColumn(IObjectNameModel* OptionalNameModel, FGetObjectClass GetObjectClassDelegate)
 	{
 		class FLabelColumn_Object : public IObjectTreeColumn
 		{
 		public:
 			
-			FLabelColumn_Object(TSharedRef<IReplicationStreamModel> Model, IObjectNameModel* OptionalNameModel)
-				: Model(MoveTemp(Model))
-				, OptionalNameModel(OptionalNameModel)
+			FLabelColumn_Object(IObjectNameModel* OptionalNameModel, FGetObjectClass GetObjectClassDelegate)
+				:  OptionalNameModel(OptionalNameModel)
+				, GetObjectClassDelegate(MoveTemp(GetObjectClassDelegate))
 			{}
 
 			virtual SHeaderRow::FColumn::FArguments CreateHeaderRowArgs() const override
@@ -49,10 +49,12 @@ namespace UE::ConcertSharedSlate::ReplicationColumns::TopLevel
 			virtual TSharedRef<SWidget> GenerateColumnWidget(const FBuildArgs& InArgs) override
 			{
 				const FReplicatedObjectData& ObjectData = InArgs.RowItem.RowData;
+				const FSoftObjectPath& ObjectPath = ObjectData.GetObjectPath();
+				
 				const FText Text = GetDisplayText(ObjectData);
-					
+				const FSlateIcon ClassIcon = GetObjectClassDelegate.IsBound() ? DisplayUtils::GetObjectIcon(GetObjectClassDelegate.Execute(ObjectPath)) : FSlateIcon{};
 				return SNew(SHorizontalBox)
-					.ToolTipText(FText::FromString(ObjectData.GetObjectPath().ToString()))
+					.ToolTipText(FText::FromString(ObjectPath.ToString()))
 					
 					+SHorizontalBox::Slot()
 					.AutoWidth()
@@ -60,7 +62,7 @@ namespace UE::ConcertSharedSlate::ReplicationColumns::TopLevel
 					.VAlign(VAlign_Center)
 					[
 						SNew(SImage)
-						.Image(DisplayUtils::GetObjectIcon(*Model, ObjectData.GetObjectPath()).GetOptionalIcon())
+						.Image(ClassIcon.GetOptionalIcon())
 					]
 					
 					+SHorizontalBox::Slot()
@@ -88,8 +90,8 @@ namespace UE::ConcertSharedSlate::ReplicationColumns::TopLevel
 
 		private:
 			
-			const TSharedRef<IReplicationStreamModel> Model;
 			IObjectNameModel* const OptionalNameModel;
+			const FGetObjectClass GetObjectClassDelegate;
 			
 			FText GetDisplayText(const FReplicatedObjectData& ObjectData) const
 			{
@@ -99,23 +101,23 @@ namespace UE::ConcertSharedSlate::ReplicationColumns::TopLevel
 		};
 
 		return {
-			TReplicationColumnDelegates<FObjectTreeRowContext>::FCreateColumn::CreateLambda([Model = MoveTemp(Model), OptionalNameModel]()
+			TReplicationColumnDelegates<FObjectTreeRowContext>::FCreateColumn::CreateLambda([OptionalNameModel, GetObjectClassDelegate = MoveTemp(GetObjectClassDelegate)]()
 			{
-				return MakeShared<FLabelColumn_Object>(Model, OptionalNameModel);
+				return MakeShared<FLabelColumn_Object>(OptionalNameModel, GetObjectClassDelegate);
 			}),
 			LabelColumnId,
 			{ static_cast<int32>(ETopLevelColumnOrder::Label) }
 		};
 	}
 	
-	FObjectColumnEntry TypeColumn(TSharedRef<IReplicationStreamModel> Model)
+	FObjectColumnEntry TypeColumn(FGetObjectClass GetObjectClassDelegate)
 	{
-		class FLabelColumn_Type : public IObjectTreeColumn
+		class FTypeColumn_Object : public IObjectTreeColumn
 		{
 		public:
 			
-			FLabelColumn_Type(TSharedRef<IReplicationStreamModel> Model)
-				: Model(MoveTemp(Model))
+			FTypeColumn_Object(FGetObjectClass GetObjectClassDelegate)
+				: GetObjectClassDelegate(MoveTemp(GetObjectClassDelegate))
 			{}
 
 			virtual SHeaderRow::FColumn::FArguments CreateHeaderRowArgs() const override
@@ -127,36 +129,43 @@ namespace UE::ConcertSharedSlate::ReplicationColumns::TopLevel
 			
 			virtual TSharedRef<SWidget> GenerateColumnWidget(const FBuildArgs& InArgs) override
 			{
+				const FSoftClassPath Class = GetClass(InArgs.RowItem.RowData.GetObjectPath());
 				return SNew(SBox)
 					.Padding(8, 0, 0, 0) // So the type name text is aligned with the header column text
 					[
 						SNew(STextBlock)
 						.HighlightText(TAttribute<FText>::CreateLambda([HighlightText = InArgs.HighlightText](){ return *HighlightText; }))
-						.Text(DisplayUtils::GetObjectTypeText(*Model, InArgs.RowItem.RowData.GetObjectPath()))
+						.Text(DisplayUtils::GetObjectTypeText(Class))
 					];
 			}
 			
 			virtual void PopulateSearchString(const FObjectTreeRowContext& InItem, TArray<FString>& InOutSearchStrings) const override
 			{
-				InOutSearchStrings.Add(DisplayUtils::GetObjectTypeText(Model.Get(), InItem.RowData.GetObjectPath()).ToString());
+				const FSoftClassPath Class = GetClass(InItem.RowData.GetObjectPath());
+				InOutSearchStrings.Add(DisplayUtils::GetObjectTypeText(Class).ToString());
 			}
 			
 			virtual bool CanBeSorted() const override { return true; } 
 			virtual bool IsLessThan(const FObjectTreeRowContext& Left, const FObjectTreeRowContext& Right) const override
 			{
-				return DisplayUtils::GetObjectTypeText(*Model, Left.RowData.GetObjectPath()).ToString()
-					< DisplayUtils::GetObjectTypeText(*Model, Right.RowData.GetObjectPath()).ToString();
+				const FSoftClassPath LeftClass = GetClass(Left.RowData.GetObjectPath());
+				const FSoftClassPath RightClass = GetClass(Right.RowData.GetObjectPath());
+				return DisplayUtils::GetObjectTypeText(LeftClass).ToString()
+					< DisplayUtils::GetObjectTypeText(RightClass).ToString();
 			}
 
 		private:
 			
-			const TSharedRef<IReplicationStreamModel> Model;
+			const FGetObjectClass GetObjectClassDelegate;
+
+			FSoftClassPath GetClass(const FSoftObjectPath& Path) const { return GetObjectClassDelegate.Execute(Path); }
 		};
 
+		check(GetObjectClassDelegate.IsBound());
 		return {
-			TReplicationColumnDelegates<FObjectTreeRowContext>::FCreateColumn::CreateLambda([Model = MoveTemp(Model)]()
+			TReplicationColumnDelegates<FObjectTreeRowContext>::FCreateColumn::CreateLambda([GetObjectClassDelegate = MoveTemp(GetObjectClassDelegate)]()
 			{
-				return MakeShared<FLabelColumn_Type>(Model);
+				return MakeShared<FTypeColumn_Object>(GetObjectClassDelegate);
 			}),
 			TypeColumnId,
 			{ static_cast<int32>(ETopLevelColumnOrder::Type) }
@@ -188,7 +197,7 @@ namespace UE::ConcertSharedSlate::ReplicationColumns::Property
 			
 			virtual TSharedRef<SWidget> GenerateColumnWidget(const FBuildArgs& InArgs) override
 			{
-				const FReplicatedPropertyData& PropertyData = InArgs.RowItem.RowData;
+				const FPropertyData& PropertyData = InArgs.RowItem.RowData;
 				return SNew(STextBlock)
 					.Font(FAppStyle::GetFontStyle("PropertyWindow.NormalFont"))
 					.HighlightText(TAttribute<FText>::CreateLambda([HighlightText = InArgs.HighlightText](){ return *HighlightText; }))
@@ -197,7 +206,7 @@ namespace UE::ConcertSharedSlate::ReplicationColumns::Property
 			
 			virtual void PopulateSearchString(const FPropertyTreeRowContext& InItem, TArray<FString>& InOutSearchStrings) const override
 			{
-				const FReplicatedPropertyData& PropertyData = InItem.RowData;
+				const FPropertyData& PropertyData = InItem.RowData;
 				FString DisplayString = DisplayUtils::GetPropertyDisplayString(InItem.RowData.GetProperty(), ResolveOrLoadClass(PropertyData));
 				InOutSearchStrings.Emplace(MoveTemp(DisplayString));
 			}
@@ -211,7 +220,7 @@ namespace UE::ConcertSharedSlate::ReplicationColumns::Property
 
 		private:
 
-			static UClass* ResolveOrLoadClass(const FReplicatedPropertyData& PropertyData)
+			static UClass* ResolveOrLoadClass(const FPropertyData& PropertyData)
 			{
 #if WITH_EDITOR
 				// On editor there may be Blueprints that need loading...
@@ -248,7 +257,7 @@ namespace UE::ConcertSharedSlate::ReplicationColumns::Property
 			
 			virtual TSharedRef<SWidget> GenerateColumnWidget(const FBuildArgs& InArgs) override
 			{
-				const FReplicatedPropertyData& PropertyData = InArgs.RowItem.RowData;
+				const FPropertyData& PropertyData = InArgs.RowItem.RowData;
 				return SNew(SBox)
 					.Padding(8, 0, 0, 0) // So the type name text is aligned with the header column text
 					[
@@ -272,7 +281,7 @@ namespace UE::ConcertSharedSlate::ReplicationColumns::Property
 
 		private:
 			
-			static FText GetDisplayText(const FReplicatedPropertyData& Args)
+			static FText GetDisplayText(const FPropertyData& Args)
 			{
 				UClass* Class = Args.GetOwningClass().TryLoadClass<UObject>();
 				const FProperty* Property = Class ? ConcertSyncCore::PropertyChain::ResolveProperty(*Class, Args.GetProperty()) : nullptr;
diff --git a/Engine/Plugins/Developer/Concert/ConcertUI/ConcertSharedSlate/Source/ConcertSharedSlate/Private/Replication/Editor/View/Tree/SReplicationColumnRow.h b/Engine/Plugins/Developer/Concert/ConcertUI/ConcertSharedSlate/Source/ConcertSharedSlate/Private/Replication/Editor/View/Tree/SReplicationColumnRow.h
--- a/Engine/Plugins/Developer/Concert/ConcertUI/ConcertSharedSlate/Source/ConcertSharedSlate/Private/Replication/Editor/View/Tree/SReplicationColumnRow.h
+++ b/Engine/Plugins/Developer/Concert/ConcertUI/ConcertSharedSlate/Source/ConcertSharedSlate/Private/Replication/Editor/View/Tree/SReplicationColumnRow.h
@@ -3,6 +3,7 @@
 #pragma once
 
 #include "ConcertFrontendStyle.h"
+#include "TreeItemTraitsInput.h"
 
 #include "Delegates/DelegateCombinations.h"
 #include "Widgets/Layout/SBox.h"
@@ -17,18 +18,14 @@
 
 namespace UE::ConcertSharedSlate
 {
-	/**
-	 * Displays the summary of an activity recorded and recoverable in the SConcertSessionRecovery list view.
-	 */
+	/** A row that generates its columns generically via IReplicationTreeColumn. */
 	template<typename TListItemType>
 	class SReplicationColumnRow : public SMultiColumnTableRow<TSharedPtr<TListItemType>>
 	{
 	public:
-		
-		DECLARE_DELEGATE_RetVal_OneParam(TSharedPtr<IReplicationTreeColumn<TListItemType>>, FGetColumn,
-			const FName& ColumnId
-			);
-		DECLARE_DELEGATE_RetVal_TwoParams(TSharedPtr<SWidget>, FOverrideColumnWidget, const FName& ColumnName, const TListItemType& RowData);
+
+		using FGetColumn = typename TReplicationTreeData<TListItemType>::FGetColumn;
+		using FOverrideColumnWidget = typename TReplicationTreeData<TListItemType>::FOverrideColumnWidget;
 		
 		SLATE_BEGIN_ARGS(SReplicationColumnRow)
 			: _Style(&FCoreStyle::Get().GetWidgetStyle<FTableRowStyle>("TableView.Row"))
diff --git a/Engine/Plugins/Developer/Concert/ConcertUI/ConcertSharedSlate/Source/ConcertSharedSlate/Private/Replication/Editor/View/Tree/SReplicationTreeView.h b/Engine/Plugins/Developer/Concert/ConcertUI/ConcertSharedSlate/Source/ConcertSharedSlate/Private/Replication/Editor/View/Tree/SReplicationTreeView.h
--- a/Engine/Plugins/Developer/Concert/ConcertUI/ConcertSharedSlate/Source/ConcertSharedSlate/Private/Replication/Editor/View/Tree/SReplicationTreeView.h
+++ b/Engine/Plugins/Developer/Concert/ConcertUI/ConcertSharedSlate/Source/ConcertSharedSlate/Private/Replication/Editor/View/Tree/SReplicationTreeView.h
@@ -4,7 +4,8 @@
 
 #include "Replication/Editor/View/Column/ReplicationColumnInfo.h"
 #include "Replication/Editor/View/Column/ReplicationColumnsUtils.h"
-#include "SReplicationColumnRow.h"
+#include "TreeItemTraitsInput.h"
+#include "TreeItemTraits.h"
 
 #include "Algo/RemoveIf.h"
 #include "Misc/TextFilter.h"
@@ -21,22 +22,36 @@
 
 namespace UE::ConcertSharedSlate
 {
+	enum class EItemFilterResult : uint8
+	{
+		/** The item passes the filter and should be displayed */
+		Include,
+		/** The item does not pass the filter and not should be displayed */
+		Exclude,
+		/** The item is only displayed if the filter passes on one of its children. If no child is included, the item is not displayed. */
+		IncludeOnlyIfChildIsIncluded
+	};
+	
 	/**
-	 * Shared code for the list view for replicated actors and properties.
-	 * It is a table view that is searchable with a search box and exposes slots to add more filter widgets, such as SBasicFilterBar.
+	 * Shared code for the tree view for replicated actors and properties.
+	 * 
+	 * It is a tree view that is searchable with a search box and exposes slots to add more filter widgets, such as SBasicFilterBar.
+	 * The columns are modular and are abstracted by IReplicationTreeColumn. You can customize row content by by specializing TReplicationTreeItemTraits;
+	 * the default implementation creates a SReplicationColumnRow, which uses IReplicationTreeColumn to draw the columns. 
 	 */
 	template<typename TItemType>
 	class SReplicationTreeView : public SCompoundWidget
 	{
 	public:
 
-		using TOverrideColumnWidget = typename SReplicationColumnRow<TItemType>::FOverrideColumnWidget;
+		using TOverrideColumnWidget = typename TReplicationTreeData<TItemType>::FOverrideColumnWidget;
 
 		DECLARE_DELEGATE_OneParam(FDeleteItems, const TArray<TSharedPtr<TItemType>>& SelectedItems);
 		DECLARE_DELEGATE_TwoParams(FGetItemChildren, TSharedPtr<TItemType> Item, TFunctionRef<void(TSharedPtr<TItemType>)> ProcessChild);
 		DECLARE_DELEGATE(FOnSelectionChanged);
-		
-		DECLARE_DELEGATE_RetVal_OneParam(bool, FCustomFilter, const TSharedPtr<TItemType>& Item);
+
+		/** @return Whether the item passes the filter, i.e. should be shown. */
+		DECLARE_DELEGATE_RetVal_OneParam(EItemFilterResult, FCustomFilter, const TItemType& Item);
 		DECLARE_DELEGATE_RetVal_OneParam(bool, FIsSearchableItem, const TSharedPtr<TItemType>& Item);
 
 		enum class EContent
@@ -363,7 +378,7 @@ namespace UE::ConcertSharedSlate
 		void PopulateSearchStrings(const TSharedPtr<TItemType>& Item, TArray<FString>& OutSearchStrings);
 		void ReapplyFilters();
 		EFilterResult ApplyFiltersRecursive(const TSharedPtr<TItemType>& Item, TSet<TSharedPtr<TItemType>>& FilteredItemsToShow);
-		bool PassesFilters(const TSharedPtr<TItemType>& Item);
+		EItemFilterResult PassesFilters(const TSharedPtr<TItemType>& Item);
 
 		/** Called after RootItems has changed. Removes all invalidated ItemMetaData entries */
 		void CleanseItemMetaData();
@@ -489,18 +504,18 @@ namespace UE::ConcertSharedSlate
 	TSharedRef<ITableRow> SReplicationTreeView<TItemType>::OnGenerateRowWidget(TSharedPtr<TItemType> Item, const TSharedRef<STableViewBase>& OwnerTable)
 	{
 		const typename SReplicationColumnRow<TItemType>::FGetColumn ColumnGetter =
-			SReplicationColumnRow<TItemType>::FGetColumn::CreateSP(
-				this,
-				&SReplicationTreeView::FindColumnByName
+			SReplicationColumnRow<TItemType>::FGetColumn::CreateSP(this, &SReplicationTreeView::FindColumnByName
 			);
-	
-		return SNew(SReplicationColumnRow<TItemType>, OwnerTable)
-			.HighlightText(SearchText)
-			.ColumnGetter(ColumnGetter)
-			.OverrideColumnWidget(OverrideColumnWidget)
-			.RowData(Item)
-			.ExpandableColumnLabel(ExpandableColumnId)
-			.Style(RowStyle);
+
+		// Compile-time decide the row that is supposed to be generated
+		return TReplicationTreeItemTraits<TItemType>::GenerateRowWidget(Item, OwnerTable,
+			{
+				ColumnGetter,
+				OverrideColumnWidget,
+				SearchText,
+				ExpandableColumnId,
+				RowStyle
+			});
 	}
 
 	template <typename TItemType>
@@ -513,8 +528,8 @@ namespace UE::ConcertSharedSlate
 				const bool bAllowedByFilter =
 					// When we applied the filter, was this item or one of its children allowed?
 					AllFilteredItems.Contains(ItemToAdd)
-					// Handle case where no filter has been applied, yet
-					|| PassesFilters(ItemToAdd);
+					// Handle case where no filter has been applied, yet. 
+					|| PassesFilters(ItemToAdd) != EItemFilterResult::Exclude;
 				if (bAllowedByFilter)
 				{
 					OutChildren.Add(ItemToAdd);
@@ -652,28 +667,38 @@ namespace UE::ConcertSharedSlate
 	typename SReplicationTreeView<TItemType>::EFilterResult SReplicationTreeView<TItemType>::ApplyFiltersRecursive(const TSharedPtr<TItemType>& Item, TSet<TSharedPtr<TItemType>>& FilteredItemsToShow)
 	{
 		bool bPassesAtLeastOnce = false;
-		if (PassesFilters(Item))
-		{
-			bPassesAtLeastOnce = true;
-			FilteredItemsToShow.Add(Item);
-		}
+		const EItemFilterResult ItemFilterResult = PassesFilters(Item);
 
+		bool bChildPassedAtLeastOnce = false;
 		if (OnGetChildrenDelegate.IsBound())
 		{
-			OnGetChildrenDelegate.Execute(Item, [this, &FilteredItemsToShow, &bPassesAtLeastOnce](TSharedPtr<TItemType> ItemToAdd)
+			OnGetChildrenDelegate.Execute(Item, [this, &FilteredItemsToShow, &bChildPassedAtLeastOnce](TSharedPtr<TItemType> ItemToAdd)
 			{
-				bPassesAtLeastOnce |= ApplyFiltersRecursive(ItemToAdd, FilteredItemsToShow) == EFilterResult::ItemOrChildrenPassFilter;
+				bChildPassedAtLeastOnce |= ApplyFiltersRecursive(ItemToAdd, FilteredItemsToShow) == EFilterResult::ItemOrChildrenPassFilter;
 			});
+			bPassesAtLeastOnce = bChildPassedAtLeastOnce;
+		}
+		
+		if (ItemFilterResult == EItemFilterResult::Include
+			 || (ItemFilterResult == EItemFilterResult::IncludeOnlyIfChildIsIncluded && bChildPassedAtLeastOnce))
+		{
+			bPassesAtLeastOnce = true;
+			FilteredItemsToShow.Add(Item);
 		}
 
 		return bPassesAtLeastOnce ? EFilterResult::ItemOrChildrenPassFilter : EFilterResult::NoneInHierarchyPassFilter;
 	}
 
 	template <typename TItemType>
-	bool SReplicationTreeView<TItemType>::PassesFilters(const TSharedPtr<TItemType>& Item)
+	EItemFilterResult SReplicationTreeView<TItemType>::PassesFilters(const TSharedPtr<TItemType>& Item)
 	{
-		return SearchTextFilter->PassesFilter(Item)
-			&& (!CustomFilterDelegate.IsBound() || CustomFilterDelegate.Execute(Item));
+		if (SearchTextFilter->PassesFilter(Item))
+		{
+			return CustomFilterDelegate.IsBound()
+				? CustomFilterDelegate.Execute(*Item)
+				: EItemFilterResult::Include; 
+		}
+		return EItemFilterResult::Exclude;
 	}
 
 	template <typename TItemType>
@@ -746,8 +771,8 @@ namespace UE::ConcertSharedSlate
 
 			switch (SortMode)
 			{
-				case EColumnSortMode::Ascending: return Column->IsLessThan(*Left, *Right);
-				case EColumnSortMode::Descending: return Column->IsLessThan(*Right, *Left);
+				case EColumnSortMode::Ascending: return TReplicationTreeItemTraits<TItemType>::IsLessThan(*Left, *Right, *Column);
+				case EColumnSortMode::Descending: return TReplicationTreeItemTraits<TItemType>::IsLessThan(*Right, *Left, *Column);
 				case EColumnSortMode::None:
 				default: return false;
 			};
diff --git a/Engine/Plugins/Developer/Concert/ConcertUI/ConcertSharedSlate/Source/ConcertSharedSlate/Private/Replication/Editor/View/Tree/TreeItemTraits.h b/Engine/Plugins/Developer/Concert/ConcertUI/ConcertSharedSlate/Source/ConcertSharedSlate/Private/Replication/Editor/View/Tree/TreeItemTraits.h
new file mode 100644
--- /dev/null
+++ b/Engine/Plugins/Developer/Concert/ConcertUI/ConcertSharedSlate/Source/ConcertSharedSlate/Private/Replication/Editor/View/Tree/TreeItemTraits.h
@@ -0,0 +1,55 @@
+// Copyright Epic Games, Inc. All Rights Reserved.
+
+#pragma once
+
+#include "SReplicationColumnRow.h"
+#include "TreeItemTraitsInput.h"
+
+#include "Templates/SharedPointer.h"
+
+namespace UE::ConcertSharedSlate
+{
+	/** Re-usable default behavior for implementing TReplicationTreeItemTraits::GenerateRowWidget. Simply generates a SReplicationColumnRow. */
+	template<typename TItemType>
+	TSharedRef<ITableRow> GenerateRowWidget_Default(TSharedPtr<TItemType> Item, const TSharedRef<STableViewBase>& OwnerTable, const typename TReplicationTreeData<TItemType>::FGenerateRowArgs& AdditionalArgs)
+	{
+		return SNew(SReplicationColumnRow<TItemType>, OwnerTable)
+			.HighlightText(AdditionalArgs.HighlightText)
+			.ColumnGetter(AdditionalArgs.GetColumnDelegate)
+			.OverrideColumnWidget(AdditionalArgs.OverrideColumnWidgetDelegate)
+			.RowData(Item)
+			.ExpandableColumnLabel(AdditionalArgs.ExpandableColumnId)
+			.Style(AdditionalArgs.RowStyle);
+	}
+
+	/** Re-usable default behaviour for implementing sorting items by column. */
+	template<typename TItemType>
+	bool IsLessThan_Default(const TItemType& Left, const TItemType& Right, const IReplicationTreeColumn<TItemType>& SortByColumn)
+	{
+		return SortByColumn.IsLessThan(Left, Right);
+	}
+
+	/**
+	 * Allows type-specific overriding of SReplicationColumnRow behavior.
+	 * This effectively implements the Strategy design pattern.
+	 */
+	template<typename TItemType>
+	class TReplicationTreeItemTraits
+	{
+	public:
+
+		/** Called by SReplicationColumnRow to generate an item's row. */
+		static TSharedRef<ITableRow> GenerateRowWidget(TSharedPtr<TItemType> Item, const TSharedRef<STableViewBase>& OwnerTable, const typename TReplicationTreeData<TItemType>::FGenerateRowArgs& AdditionalArgs)
+		{
+			return GenerateRowWidget_Default(Item, OwnerTable, AdditionalArgs);
+		}
+
+		/** Called by SReplicationColumnRow when sorting items by a column (this is the primary or secondary). Implement this to account for custom row types, if any. You usually want to implement this if you implement GenerateRowWidget. */
+		static bool IsLessThan(const TItemType& Left, const TItemType& Right, const IReplicationTreeColumn<TItemType>& SortByColumn)
+		{
+			return IsLessThan_Default(Left, Right, SortByColumn);
+		}
+	};
+}
+
+#undef LOCTEXT_NAMESPACE
\ No newline at end of file
diff --git a/Engine/Plugins/Developer/Concert/ConcertUI/ConcertSharedSlate/Source/ConcertSharedSlate/Private/Replication/Editor/View/Tree/TreeItemTraitsInput.h b/Engine/Plugins/Developer/Concert/ConcertUI/ConcertSharedSlate/Source/ConcertSharedSlate/Private/Replication/Editor/View/Tree/TreeItemTraitsInput.h
new file mode 100644
--- /dev/null
+++ b/Engine/Plugins/Developer/Concert/ConcertUI/ConcertSharedSlate/Source/ConcertSharedSlate/Private/Replication/Editor/View/Tree/TreeItemTraitsInput.h
@@ -0,0 +1,40 @@
+// Copyright Epic Games, Inc. All Rights Reserved.
+
+#pragma once
+
+#include "Delegates/Delegate.h"
+#include "Templates/SharedPointer.h"
+
+class FName;
+class SWidget;
+
+namespace UE::ConcertSharedSlate
+{
+	template<typename TTreeItemType>
+	class IReplicationTreeColumn;
+
+	/** Holds type definitions that server as input to TReplicationTreeItemTraits (and thus should not be template specialized).  */
+	template<typename TItemType>
+	class TReplicationTreeData
+	{
+	public:
+		
+		DECLARE_DELEGATE_RetVal_OneParam(TSharedPtr<IReplicationTreeColumn<TItemType>>, FGetColumn, const FName& ColumnId);
+		DECLARE_DELEGATE_RetVal_TwoParams(TSharedPtr<SWidget>, FOverrideColumnWidget, const FName& ColumnName, const TItemType& RowData);
+		
+		struct FGenerateRowArgs
+		{
+			/** Gets info about a replication column. */
+			FGetColumn GetColumnDelegate;
+			/** Overrides a column's content widget. */
+			FOverrideColumnWidget OverrideColumnWidgetDelegate;
+
+			/** The text to highlight - equal to search text. */
+			TSharedPtr<FText> HighlightText;
+			/** The name of the column which will have the SExpandableArrow widget for the tree view. */
+			FName ExpandableColumnId;
+			/** Style to use for rows */
+			const FTableRowStyle* RowStyle = nullptr;
+		};
+	};
+}
\ No newline at end of file
diff --git a/Engine/Plugins/Developer/Concert/ConcertUI/ConcertSharedSlate/Source/ConcertSharedSlate/Private/Replication/ObjectUtils.cpp b/Engine/Plugins/Developer/Concert/ConcertUI/ConcertSharedSlate/Source/ConcertSharedSlate/Private/Replication/ObjectUtils.cpp
deleted file mode 100644
diff --git a/Engine/Plugins/Developer/Concert/ConcertUI/ConcertSharedSlate/Source/ConcertSharedSlate/Private/Replication/ReplicationWidgetFactories.cpp b/Engine/Plugins/Developer/Concert/ConcertUI/ConcertSharedSlate/Source/ConcertSharedSlate/Private/Replication/ReplicationWidgetFactories.cpp
--- a/Engine/Plugins/Developer/Concert/ConcertUI/ConcertSharedSlate/Source/ConcertSharedSlate/Private/Replication/ReplicationWidgetFactories.cpp
+++ b/Engine/Plugins/Developer/Concert/ConcertUI/ConcertSharedSlate/Source/ConcertSharedSlate/Private/Replication/ReplicationWidgetFactories.cpp
@@ -5,7 +5,10 @@
 #include "Editor/Model/GenericReplicationStreamModel.h"
 #include "Editor/View/MultiEditor/SMultiReplicationStreamEditor.h"
 #include "Editor/View/ObjectEditor/SBaseReplicationStreamEditor.h"
-#include "Editor/View/ObjectViewer/Property/SPropertyTreeView.h"
+#include "Editor/View/Property/SPerObjectPropertyAssignment.h"
+#include "Editor/View/Property/SPropertyTreeView.h"
+#include "Replication/PropertyAssignmentViewFactory.h"
+#include "Replication/Editor/View/IPropertyAssignmentView.h"
 
 namespace UE::ConcertSharedSlate
 {
@@ -20,7 +23,7 @@ namespace UE::ConcertSharedSlate
 	TSharedRef<IReplicationStreamEditor> CreateBaseStreamEditor(FCreateEditorParams EditorParams, FCreateViewerParams ViewerParams)
 	{
 		return SNew(SBaseReplicationStreamEditor, MoveTemp(EditorParams.DataModel), MoveTemp(EditorParams.ObjectSource), MoveTemp(EditorParams.PropertySource))
-			.PropertyTreeView(MoveTemp(ViewerParams.PropertyTreeView))
+			.PropertyAssignmentView(MoveTemp(ViewerParams.PropertyAssignmentView))
 			.ObjectColumns(MoveTemp(ViewerParams.ObjectColumns))
 			.PrimaryObjectSort(ViewerParams.PrimaryObjectSort)
 			.SecondaryObjectSort(ViewerParams.SecondaryObjectSort)
@@ -45,12 +48,14 @@ namespace UE::ConcertSharedSlate
 			Params.PropertyColumns.Add(ReplicationColumns::Property::LabelColumn());
 		}
 
-		SReplicationTreeView<FReplicatedPropertyData>::FCustomFilter FilterDelegate = Params.FilterItem.IsBound()
-			? SReplicationTreeView<FReplicatedPropertyData>::FCustomFilter::CreateLambda([Filter = MoveTemp(Params.FilterItem)](const TSharedPtr<FReplicatedPropertyData>& Item)
+		SReplicationTreeView<FPropertyData>::FCustomFilter FilterDelegate = Params.FilterItem.IsBound()
+			? SReplicationTreeView<FPropertyData>::FCustomFilter::CreateLambda([Filter = MoveTemp(Params.FilterItem)](const FPropertyData& Item)
 			{
-				return Filter.Execute(*Item.Get()) == EFilterResult::PassesFilter;
+				return Filter.Execute(Item) == EFilterResult::PassesFilter
+					? EItemFilterResult::Include
+					: EItemFilterResult::Exclude;
 			})
-			: SReplicationTreeView<FReplicatedPropertyData>::FCustomFilter{};
+			: SReplicationTreeView<FPropertyData>::FCustomFilter{};
 		
 		return SNew(SPropertyTreeView)
 			.FilterItem(MoveTemp(FilterDelegate))
@@ -65,6 +70,12 @@ namespace UE::ConcertSharedSlate
 			.NoItemsContent() [ MoveTemp(Params.NoItemsContent.Widget) ];
 	}
 
+	TSharedRef<IPropertyAssignmentView> CreatePerObjectAssignmentView(FCreatePerObjectAssignmentViewParams Params)
+	{
+		return SNew(SPerObjectPropertyAssignment, Params.PropertyTreeView)
+			.PropertySource(Params.PropertySource);
+	}
+
 	TSharedRef<IMultiReplicationStreamEditor> CreateBaseMultiStreamEditor(FCreateMultiStreamEditorParams EditorParams, FCreateViewerParams ViewerParams)
 	{
 		return SNew(SMultiReplicationStreamEditor, MoveTemp(EditorParams), MoveTemp(ViewerParams));
diff --git a/Engine/Plugins/Developer/Concert/ConcertUI/ConcertSharedSlate/Source/ConcertSharedSlate/Public/Replication/Editor/Model/Data/PropertyData.h b/Engine/Plugins/Developer/Concert/ConcertUI/ConcertSharedSlate/Source/ConcertSharedSlate/Public/Replication/Editor/Model/Data/PropertyData.h
new file mode 100644
--- /dev/null
+++ b/Engine/Plugins/Developer/Concert/ConcertUI/ConcertSharedSlate/Source/ConcertSharedSlate/Public/Replication/Editor/Model/Data/PropertyData.h
@@ -0,0 +1,30 @@
+// Copyright Epic Games, Inc. All Rights Reserved.
+
+#pragma once
+
+#include "Replication/Data/ConcertPropertySelection.h"
+
+namespace UE::ConcertSharedSlate
+{
+	/** Instanced for each property row in IPropertyTreeView.*/
+	class FPropertyData
+	{
+	public:
+		
+		FPropertyData(FSoftClassPath OwningClass, FConcertPropertyChain Object)
+			: OwningClass(MoveTemp(OwningClass))
+			, Property(MoveTemp(Object))
+		{}
+		
+		const FConcertPropertyChain& GetProperty() const { return Property; }
+		const FSoftClassPath& GetOwningClass() const { return OwningClass; }
+
+	private:
+
+		/** The class with which the FProperty can be determined. */
+		FSoftClassPath OwningClass;
+		
+		/** The property to be replicated */
+		FConcertPropertyChain Property;
+	};
+}
\ No newline at end of file
diff --git a/Engine/Plugins/Developer/Concert/ConcertUI/ConcertSharedSlate/Source/ConcertSharedSlate/Public/Replication/Editor/Model/Data/ReplicatedObjectData.h b/Engine/Plugins/Developer/Concert/ConcertUI/ConcertSharedSlate/Source/ConcertSharedSlate/Public/Replication/Editor/Model/Data/ReplicatedObjectData.h
new file mode 100644
--- /dev/null
+++ b/Engine/Plugins/Developer/Concert/ConcertUI/ConcertSharedSlate/Source/ConcertSharedSlate/Public/Replication/Editor/Model/Data/ReplicatedObjectData.h
@@ -0,0 +1,25 @@
+// Copyright Epic Games, Inc. All Rights Reserved.
+
+#pragma once
+
+#include "UObject/SoftObjectPath.h"
+
+namespace UE::ConcertSharedSlate
+{
+	/** Instanced for each object row in stream viewer / editor UI.*/
+	class FReplicatedObjectData
+	{
+	public:
+
+		FReplicatedObjectData(FSoftObjectPath ObjectPath)
+			: ObjectPath(MoveTemp(ObjectPath))
+		{}
+		
+		const FSoftObjectPath& GetObjectPath() const { return ObjectPath; }
+
+	private:
+
+		/** The replicated object */
+		FSoftObjectPath ObjectPath;
+	};
+}
\ No newline at end of file
diff --git a/Engine/Plugins/Developer/Concert/ConcertUI/ConcertSharedSlate/Source/ConcertSharedSlate/Public/Replication/Editor/Model/ReplicatedObjectData.h b/Engine/Plugins/Developer/Concert/ConcertUI/ConcertSharedSlate/Source/ConcertSharedSlate/Public/Replication/Editor/Model/ReplicatedObjectData.h
deleted file mode 100644
diff --git a/Engine/Plugins/Developer/Concert/ConcertUI/ConcertSharedSlate/Source/ConcertSharedSlate/Public/Replication/Editor/Model/ReplicatedPropertyData.h b/Engine/Plugins/Developer/Concert/ConcertUI/ConcertSharedSlate/Source/ConcertSharedSlate/Public/Replication/Editor/Model/ReplicatedPropertyData.h
deleted file mode 100644
diff --git a/Engine/Plugins/Developer/Concert/ConcertUI/ConcertSharedSlate/Source/ConcertSharedSlate/Public/Replication/Editor/View/Column/IObjectTreeColumn.h b/Engine/Plugins/Developer/Concert/ConcertUI/ConcertSharedSlate/Source/ConcertSharedSlate/Public/Replication/Editor/View/Column/IObjectTreeColumn.h
--- a/Engine/Plugins/Developer/Concert/ConcertUI/ConcertSharedSlate/Source/ConcertSharedSlate/Public/Replication/Editor/View/Column/IObjectTreeColumn.h
+++ b/Engine/Plugins/Developer/Concert/ConcertUI/ConcertSharedSlate/Source/ConcertSharedSlate/Public/Replication/Editor/View/Column/IObjectTreeColumn.h
@@ -4,7 +4,7 @@
 
 #include "IReplicationTreeColumn.h"
 #include "ReplicationColumnInfo.h"
-#include "Replication/Editor/Model/ReplicatedObjectData.h"
+#include "Replication/Editor/Model/Data/ReplicatedObjectData.h"
 
 namespace UE::ConcertSharedSlate
 {
diff --git a/Engine/Plugins/Developer/Concert/ConcertUI/ConcertSharedSlate/Source/ConcertSharedSlate/Public/Replication/Editor/View/Column/IPropertyTreeColumn.h b/Engine/Plugins/Developer/Concert/ConcertUI/ConcertSharedSlate/Source/ConcertSharedSlate/Public/Replication/Editor/View/Column/IPropertyTreeColumn.h
--- a/Engine/Plugins/Developer/Concert/ConcertUI/ConcertSharedSlate/Source/ConcertSharedSlate/Public/Replication/Editor/View/Column/IPropertyTreeColumn.h
+++ b/Engine/Plugins/Developer/Concert/ConcertUI/ConcertSharedSlate/Source/ConcertSharedSlate/Public/Replication/Editor/View/Column/IPropertyTreeColumn.h
@@ -4,7 +4,7 @@
 
 #include "IReplicationTreeColumn.h"
 #include "ReplicationColumnInfo.h"
-#include "Replication/Editor/Model/ReplicatedPropertyData.h"
+#include "Replication/Editor/Model/Data/PropertyData.h"
 
 namespace UE::ConcertSharedSlate
 {
@@ -12,7 +12,7 @@ namespace UE::ConcertSharedSlate
 	struct FPropertyTreeRowContext
 	{
 		/** The data that is stored in the row */
-		FReplicatedPropertyData RowData;
+		FPropertyData RowData;
 	};
 	
 	using IPropertyTreeColumn = IReplicationTreeColumn<FPropertyTreeRowContext>;
diff --git a/Engine/Plugins/Developer/Concert/ConcertUI/ConcertSharedSlate/Source/ConcertSharedSlate/Public/Replication/Editor/View/Column/IReplicationTreeColumn.h b/Engine/Plugins/Developer/Concert/ConcertUI/ConcertSharedSlate/Source/ConcertSharedSlate/Public/Replication/Editor/View/Column/IReplicationTreeColumn.h
--- a/Engine/Plugins/Developer/Concert/ConcertUI/ConcertSharedSlate/Source/ConcertSharedSlate/Public/Replication/Editor/View/Column/IReplicationTreeColumn.h
+++ b/Engine/Plugins/Developer/Concert/ConcertUI/ConcertSharedSlate/Source/ConcertSharedSlate/Public/Replication/Editor/View/Column/IReplicationTreeColumn.h
@@ -19,7 +19,7 @@ namespace UE::ConcertSharedSlate
 		struct FBuildArgs
 		{
 			TSharedPtr<FText> HighlightText;
-			TTreeItemType RowItem;
+			const TTreeItemType& RowItem;
 		};
 		
 		virtual SHeaderRow::FColumn::FArguments CreateHeaderRowArgs() const = 0;
diff --git a/Engine/Plugins/Developer/Concert/ConcertUI/ConcertSharedSlate/Source/ConcertSharedSlate/Public/Replication/Editor/View/Column/SelectionViewerColumns.h b/Engine/Plugins/Developer/Concert/ConcertUI/ConcertSharedSlate/Source/ConcertSharedSlate/Public/Replication/Editor/View/Column/SelectionViewerColumns.h
--- a/Engine/Plugins/Developer/Concert/ConcertUI/ConcertSharedSlate/Source/ConcertSharedSlate/Public/Replication/Editor/View/Column/SelectionViewerColumns.h
+++ b/Engine/Plugins/Developer/Concert/ConcertUI/ConcertSharedSlate/Source/ConcertSharedSlate/Public/Replication/Editor/View/Column/SelectionViewerColumns.h
@@ -4,6 +4,7 @@
 
 #include "IObjectTreeColumn.h"
 #include "IPropertyTreeColumn.h"
+#include "Replication/Utils/ReplicationWidgetDelegates.h"
 #include "ReplicationColumnInfo.h"
 
 #include "Templates/SharedPointer.h"
@@ -32,8 +33,12 @@ namespace UE::ConcertSharedSlate::ReplicationColumns::TopLevel
 		Type = 30,
 	};
 
-	CONCERTSHAREDSLATE_API FObjectColumnEntry LabelColumn(TSharedRef<IReplicationStreamModel> Model, IObjectNameModel* OptionalNameModel = nullptr);
-	CONCERTSHAREDSLATE_API FObjectColumnEntry TypeColumn(TSharedRef<IReplicationStreamModel> Model);
+	/**
+	 * @param OptionalNameModel Used to look object name. Defaults to name displayed in object path if unset.
+	 * @param GetObjectClassDelegate Used to display class icon. No icon is displayed if unset.
+	 */
+	CONCERTSHAREDSLATE_API FObjectColumnEntry LabelColumn(IObjectNameModel* OptionalNameModel = nullptr, FGetObjectClass GetObjectClassDelegate = {});
+	CONCERTSHAREDSLATE_API FObjectColumnEntry TypeColumn(FGetObjectClass GetObjectClassDelegate);
 }
 
 namespace UE::ConcertSharedSlate::ReplicationColumns::Property
diff --git a/Engine/Plugins/Developer/Concert/ConcertUI/ConcertSharedSlate/Source/ConcertSharedSlate/Public/Replication/Editor/View/IPropertyAssignmentView.h b/Engine/Plugins/Developer/Concert/ConcertUI/ConcertSharedSlate/Source/ConcertSharedSlate/Public/Replication/Editor/View/IPropertyAssignmentView.h
new file mode 100644
--- /dev/null
+++ b/Engine/Plugins/Developer/Concert/ConcertUI/ConcertSharedSlate/Source/ConcertSharedSlate/Public/Replication/Editor/View/IPropertyAssignmentView.h
@@ -0,0 +1,54 @@
+// Copyright Epic Games, Inc. All Rights Reserved.
+
+#pragma once
+
+#include "Containers/ContainersFwd.h"
+#include "Containers/Set.h"
+#include "Templates/SharedPointer.h"
+
+class SWidget;
+struct FSoftObjectPath;
+
+namespace UE::ConcertSharedSlate
+{
+	class IReplicationStreamModel;
+	
+	/**
+	 * A replication assignment view displays an object's properties.
+	 * It is the piece of UI that is displayed on the bottom part of the replication viewer.
+	 * 
+	 * The replication stream viewer / editor looks like this (@see CreateBaseStreamEditor):
+	 * - Object tree view: User can click on an object.
+	 * - IPropertyAssignmentView: The clicked object's properties are displayed (the "root objects").
+	 *
+	 * Right now there is only 1 implementation: SPerObjectPropertyAssignment, which displays the root object's properties only.
+	 * We could also add an implementation that also displays the subobject properties, similar how the editor's property matrix does it. 
+	 * Generally, you can imagine the view as a tree view which has columns that can be injected (e.g. @see CreateBaseStreamEditor).
+	 */
+	class CONCERTSHAREDSLATE_API IPropertyAssignmentView
+	{
+	public:
+
+		/**
+		 * Rebuilds all displayed data immediately.
+		 *
+		 * @param Objects The objects that are supposed to be displayed.
+		 * @param Model The model that can be queried for object info.
+		 */
+		virtual void RefreshData(const TSet<FSoftObjectPath>& Objects, const IReplicationStreamModel& Model) = 0;
+		
+		/** Reapply the filter function to all items at the end of the frame. Call e.g. when the filters have changed. */
+		virtual void RequestRefilter() const = 0;
+		
+		/**
+		 * Requests that the given column be resorted, if it currently affects the row sorting (primary or secondary).
+		 * Call e.g. when a sortable attribute of the column has changed.
+		 */
+		virtual void RequestResortForColumn(const FName& ColumnId) = 0;
+
+		/** Gets the tree view's widget */
+		virtual TSharedRef<SWidget> GetWidget() = 0;
+
+		virtual ~IPropertyAssignmentView() = default;
+	};
+}
diff --git a/Engine/Plugins/Developer/Concert/ConcertUI/ConcertSharedSlate/Source/ConcertSharedSlate/Public/Replication/Editor/View/IPropertyTreeView.h b/Engine/Plugins/Developer/Concert/ConcertUI/ConcertSharedSlate/Source/ConcertSharedSlate/Public/Replication/Editor/View/IPropertyTreeView.h
--- a/Engine/Plugins/Developer/Concert/ConcertUI/ConcertSharedSlate/Source/ConcertSharedSlate/Public/Replication/Editor/View/IPropertyTreeView.h
+++ b/Engine/Plugins/Developer/Concert/ConcertUI/ConcertSharedSlate/Source/ConcertSharedSlate/Public/Replication/Editor/View/IPropertyTreeView.h
@@ -12,7 +12,7 @@ struct FSoftClassPath;
 
 namespace UE::ConcertSharedSlate
 {
-	/** Represents a tree view displaying properties. */
+	/** Represents a tree view displaying properties from a single class. */
 	class CONCERTSHAREDSLATE_API IPropertyTreeView
 	{
 	public:
@@ -28,9 +28,7 @@ namespace UE::ConcertSharedSlate
 		 */
 		virtual void RefreshPropertyData(const TSet<FConcertPropertyChain>& PropertiesToDisplay, const FSoftClassPath& Class, bool bCanReuseExistingRowItems = true) = 0;
 		
-		/**
-		 * Reapply the filter function to all items at the end of the frame. Call e.g. when the filters have changed.
-		 */
+		/** Reapply the filter function to all items at the end of the frame. Call e.g. when the filters have changed. */
 		virtual void RequestRefilter() const = 0;
 		
 		/**
diff --git a/Engine/Plugins/Developer/Concert/ConcertUI/ConcertSharedSlate/Source/ConcertSharedSlate/Public/Replication/Editor/View/IReplicationStreamViewer.h b/Engine/Plugins/Developer/Concert/ConcertUI/ConcertSharedSlate/Source/ConcertSharedSlate/Public/Replication/Editor/View/IReplicationStreamViewer.h
--- a/Engine/Plugins/Developer/Concert/ConcertUI/ConcertSharedSlate/Source/ConcertSharedSlate/Public/Replication/Editor/View/IReplicationStreamViewer.h
+++ b/Engine/Plugins/Developer/Concert/ConcertUI/ConcertSharedSlate/Source/ConcertSharedSlate/Public/Replication/Editor/View/IReplicationStreamViewer.h
@@ -34,7 +34,7 @@ namespace UE::ConcertSharedSlate
 		 * If there is an IReplicationSubobjectView, this is IReplicationSubobjectView::GetSelectedObjects.
 		 * Otherwise it is IReplicationStreamViewer::GetSelectedTopLevelObjects.
 		 */
-		virtual TArray<FSoftObjectPath> GetObjectsBeingPropertyEdited() const = 0;
+		virtual TSet<FSoftObjectPath> GetObjectsBeingPropertyEdited() const = 0;
 
 		virtual ~IReplicationStreamViewer() = default;
 	};
diff --git a/Engine/Plugins/Developer/Concert/ConcertUI/ConcertSharedSlate/Source/ConcertSharedSlate/Public/Replication/ObjectUtils.h b/Engine/Plugins/Developer/Concert/ConcertUI/ConcertSharedSlate/Source/ConcertSharedSlate/Public/Replication/ObjectUtils.h
deleted file mode 100644
diff --git a/Engine/Plugins/Developer/Concert/ConcertUI/ConcertSharedSlate/Source/ConcertSharedSlate/Public/Replication/PropertyAssignmentViewFactory.h b/Engine/Plugins/Developer/Concert/ConcertUI/ConcertSharedSlate/Source/ConcertSharedSlate/Public/Replication/PropertyAssignmentViewFactory.h
new file mode 100644
--- /dev/null
+++ b/Engine/Plugins/Developer/Concert/ConcertUI/ConcertSharedSlate/Source/ConcertSharedSlate/Public/Replication/PropertyAssignmentViewFactory.h
@@ -0,0 +1,42 @@
+// Copyright Epic Games, Inc. All Rights Reserved.
+
+#pragma once
+
+#include "PropertyTreeFactory.h"
+#include "Editor/View/Column/IPropertyTreeColumn.h"
+#include "Replication/Utils/FilterResult.h"
+
+#include "Delegates/Delegate.h"
+#include "Misc/Attribute.h"
+#include "Templates/SharedPointer.h"
+
+namespace UE::ConcertSharedSlate
+{
+	class FPropertyData;
+	class IPropertyAssignmentView;
+	class IPropertySelectionSourceModel;
+	class IPropertyTreeView;
+	
+	DECLARE_DELEGATE_RetVal_OneParam(EFilterResult, FFilterPropertyData, const FPropertyData&);
+	
+	struct FCreatePerObjectAssignmentViewParams
+	{
+		/** Required. Displays the properties in a tree view. You can pass in e.g. custom UI with advanced filtering. */
+		TSharedRef<IPropertyTreeView> PropertyTreeView = CreateSearchablePropertyTreeView();
+
+		/**
+		 * Optional.
+		 * If specified, the view will display all properties reported by the model (useful for editor UI which edits streams, not useful for server where property info is not available).
+		 * If unspecified, only display the properties assigned to the object in the stream.
+		 */
+		TSharedPtr<IPropertySelectionSourceModel> PropertySource;
+	};
+	
+	/**
+	 * Creates a view that shows the properties of the object the user clicks.
+	 * 
+	 * This basically just wraps IPropertyTreeView.
+	 * You can customize the tree view by injecting columns (@see CreateSearchablePropertyTreeView).
+	 */
+	CONCERTSHAREDSLATE_API TSharedRef<IPropertyAssignmentView> CreatePerObjectAssignmentView(FCreatePerObjectAssignmentViewParams Params = {});
+}
diff --git a/Engine/Plugins/Developer/Concert/ConcertUI/ConcertSharedSlate/Source/ConcertSharedSlate/Public/Replication/PropertyTreeFactory.h b/Engine/Plugins/Developer/Concert/ConcertUI/ConcertSharedSlate/Source/ConcertSharedSlate/Public/Replication/PropertyTreeFactory.h
new file mode 100644
--- /dev/null
+++ b/Engine/Plugins/Developer/Concert/ConcertUI/ConcertSharedSlate/Source/ConcertSharedSlate/Public/Replication/PropertyTreeFactory.h
@@ -0,0 +1,53 @@
+// Copyright Epic Games, Inc. All Rights Reserved.
+
+#pragma once
+
+#include "Editor/View/Column/IPropertyTreeColumn.h"
+#include "Editor/View/Column/ReplicationColumnsUtils.h"
+#include "Editor/View/Column/SelectionViewerColumns.h"
+#include "Replication/Utils/FilterResult.h"
+
+#include "Delegates/Delegate.h"
+#include "Misc/Attribute.h"
+#include "Templates/SharedPointer.h"
+
+namespace UE::ConcertSharedSlate
+{
+	class FPropertyData;
+	class IPropertyTreeView;
+	
+	DECLARE_DELEGATE_RetVal_OneParam(EFilterResult, FFilterPropertyData, const FPropertyData&);
+	
+	struct FCreatePropertyTreeViewParams
+	{
+		/** Optional. Additional property columns you want added. */
+		TArray<FPropertyColumnEntry> PropertyColumns
+		{
+			ReplicationColumns::Property::LabelColumn(),
+			ReplicationColumns::Property::TypeColumn()
+		};
+
+		/** Optional filter function. Return true to al */
+		FFilterPropertyData FilterItem;
+		
+		/** Optional initial primary sort mode for object rows */
+		FColumnSortInfo PrimaryPropertySort { ReplicationColumns::Property::LabelColumnId, EColumnSortMode::Ascending };
+		/** Optional initial secondary sort mode for object rows */
+		FColumnSortInfo SecondaryPropertySort { ReplicationColumns::Property::LabelColumnId, EColumnSortMode::Ascending };
+		
+		/** Optional widget to add to the left of the property list search bar. */
+		TAlwaysValidWidget LeftOfPropertySearchBar;
+		/** Optional widget to add to the right of the property list search bar. */
+		TAlwaysValidWidget RightOfPropertySearchBar;
+		/** Optional widget to add between the search bar and the table view (e.g. a SBasicFilterBar). */
+		TAlwaysValidWidget RowBelowSearchBar;
+		/** Optional, alternate content to show instead of the tree view when there are no rows. */
+		TAlwaysValidWidget NoItemsContent;
+	};
+	
+	/**
+	 * Creates a tree view that uses a search box for filtering items.
+	 * You can customize this tree view by adding custom widgets and columns into the property view.
+	 */
+	CONCERTSHAREDSLATE_API TSharedRef<IPropertyTreeView> CreateSearchablePropertyTreeView(FCreatePropertyTreeViewParams Params = {});
+}
diff --git a/Engine/Plugins/Developer/Concert/ConcertUI/ConcertSharedSlate/Source/ConcertSharedSlate/Public/Replication/ReplicationWidgetDelegates.h b/Engine/Plugins/Developer/Concert/ConcertUI/ConcertSharedSlate/Source/ConcertSharedSlate/Public/Replication/ReplicationWidgetDelegates.h
deleted file mode 100644
diff --git a/Engine/Plugins/Developer/Concert/ConcertUI/ConcertSharedSlate/Source/ConcertSharedSlate/Public/Replication/ReplicationWidgetFactories.h b/Engine/Plugins/Developer/Concert/ConcertUI/ConcertSharedSlate/Source/ConcertSharedSlate/Public/Replication/ReplicationWidgetFactories.h
--- a/Engine/Plugins/Developer/Concert/ConcertUI/ConcertSharedSlate/Source/ConcertSharedSlate/Public/Replication/ReplicationWidgetFactories.h
+++ b/Engine/Plugins/Developer/Concert/ConcertUI/ConcertSharedSlate/Source/ConcertSharedSlate/Public/Replication/ReplicationWidgetFactories.h
@@ -2,27 +2,30 @@
 
 #pragma once
 
+#include "PropertyAssignmentViewFactory.h"
 #include "Editor/View/Column/IObjectTreeColumn.h"
-#include "Editor/View/Column/IPropertyTreeColumn.h"
 #include "Editor/View/Column/ReplicationColumnInfo.h"
 #include "Editor/View/Column/SelectionViewerColumns.h"
+#include "PropertyTreeFactory.h"
 #include "Replication/Editor/View/Column/ReplicationColumnsUtils.h"
-#include "ReplicationWidgetDelegates.h"
+#include "Replication/Utils/ReplicationWidgetDelegates.h"
 
 #include "Delegates/Delegate.h"
 #include "Misc/Attribute.h"
 #include "Templates/SharedPointer.h"
 
+namespace UE::ConcertSharedSlate
+{
+	class IPropertyAssignmentView;
+}
+
 class UObject;
-class SWidget;
 
 struct FConcertStreamObjectAutoBindingRules;
 struct FConcertObjectReplicationMap;
 
 namespace UE::ConcertSharedSlate
 {
-	class FReplicatedPropertyData;
-	
 	class IEditableMultiReplicationStreamModel;
 	class IEditableReplicationStreamModel;
 	class IMultiReplicationStreamEditor;
@@ -32,7 +35,6 @@ namespace UE::ConcertSharedSlate
 	class IReplicationStreamModel;
 	class IReplicationStreamEditor;
 	class IReplicationStreamViewer;
-	class IReplicationSubobjectView;
 	class IStreamExtender;
 	class IObjectHierarchyModel;
 	class IPropertySelectionSourceModel;
@@ -49,54 +51,15 @@ namespace UE::ConcertSharedSlate
 		TAttribute<FConcertObjectReplicationMap*> ReplicationMapAttribute,
 		TSharedPtr<IStreamExtender> Extender = nullptr
 		);
-
-	enum class EFilterResult : uint8
-	{
-		/** Include the object into list of displayed objects */
-		PassesFilter,
-		/** Exclude the object from list of displayed objects */
-		DoesNotPassFilter
-	};
-	DECLARE_DELEGATE_RetVal_OneParam(EFilterResult, FFilterPropertyData, const FReplicatedPropertyData&);
-	
-	struct FCreatePropertyTreeViewParams
-	{
-		/** Optional. Additional property columns you want added. */
-		TArray<FPropertyColumnEntry> PropertyColumns
-		{
-			ReplicationColumns::Property::LabelColumn(),
-			ReplicationColumns::Property::TypeColumn()
-		};
-
-		/** Optional filter function. Return true to al */
-		FFilterPropertyData FilterItem;
-		
-		/** Optional initial primary sort mode for object rows */
-		FColumnSortInfo PrimaryPropertySort { ReplicationColumns::Property::LabelColumnId, EColumnSortMode::Ascending };
-		/** Optional initial secondary sort mode for object rows */
-		FColumnSortInfo SecondaryPropertySort { ReplicationColumns::Property::LabelColumnId, EColumnSortMode::Ascending };
-		
-		/** Optional widget to add to the left of the property list search bar. */
-		TAlwaysValidWidget LeftOfPropertySearchBar;
-		/** Optional widget to add to the right of the property list search bar. */
-		TAlwaysValidWidget RightOfPropertySearchBar;
-		/** Optional widget to add between the search bar and the table view (e.g. a SBasicFilterBar). */
-		TAlwaysValidWidget RowBelowSearchBar;
-		/** Optional, alternate content to show instead of the tree view when there are no rows. */
-		TAlwaysValidWidget NoItemsContent;
-	};
-	
-	/**
-	 * Creates a tree view that uses a search box for filtering items.
-	 * You can customize this tree view by adding custom widgets and columns into the property view.
-	 */
-	CONCERTSHAREDSLATE_API TSharedRef<IPropertyTreeView> CreateSearchablePropertyTreeView(FCreatePropertyTreeViewParams Params = {});
 	
 	/** Params for creating a IReplicationStreamViewer. */
 	struct FCreateViewerParams
 	{
 		/** Required. Displays the properties in a tree view. You can pass in e.g. custom UI with advanced filtering. */
 		TSharedRef<IPropertyTreeView> PropertyTreeView = CreateSearchablePropertyTreeView();
+
+		/** Required. In the lower half of the editor, this view presents the properties associated with the object that is currently selected in the upper part of the view. */
+		TSharedRef<IPropertyAssignmentView> PropertyAssignmentView = CreatePerObjectAssignmentView(); // TODO DP: Change this to the matrix view
 		
 		/**
 		 * Optional. Determines the objects displayed as children to the top-level objects in the top section.
diff --git a/Engine/Plugins/Developer/Concert/ConcertUI/ConcertSharedSlate/Source/ConcertSharedSlate/Public/Replication/Utils/FilterResult.h b/Engine/Plugins/Developer/Concert/ConcertUI/ConcertSharedSlate/Source/ConcertSharedSlate/Public/Replication/Utils/FilterResult.h
new file mode 100644
--- /dev/null
+++ b/Engine/Plugins/Developer/Concert/ConcertUI/ConcertSharedSlate/Source/ConcertSharedSlate/Public/Replication/Utils/FilterResult.h
@@ -0,0 +1,14 @@
+// Copyright Epic Games, Inc. All Rights Reserved.
+
+#pragma once
+
+#include "HAL/Platform.h"
+
+namespace UE::ConcertSharedSlate
+{
+	enum class EFilterResult : uint8
+	{
+		PassesFilter,
+		DoesNotPassFilter
+	};
+}
diff --git a/Engine/Plugins/Developer/Concert/ConcertUI/ConcertSharedSlate/Source/ConcertSharedSlate/Public/Replication/Utils/ObjectUtils.h b/Engine/Plugins/Developer/Concert/ConcertUI/ConcertSharedSlate/Source/ConcertSharedSlate/Public/Replication/Utils/ObjectUtils.h
new file mode 100644
--- /dev/null
+++ b/Engine/Plugins/Developer/Concert/ConcertUI/ConcertSharedSlate/Source/ConcertSharedSlate/Public/Replication/Utils/ObjectUtils.h
@@ -0,0 +1,19 @@
+// Copyright Epic Games, Inc. All Rights Reserved.
+
+#pragma once
+
+#include "Misc/Optional.h"
+
+struct FSoftObjectPath;
+
+namespace UE::ConcertSharedSlate::ObjectUtils
+{
+	/** @return Checks whether the given object is an actor. */
+	CONCERTSHAREDSLATE_API bool IsActor(const FSoftObjectPath& Object);
+	
+	/** @return Get the owning actor of Subobject. If Subobject is an actor, then this returns unset. */
+	CONCERTSHAREDSLATE_API TOptional<FSoftObjectPath> GetActorOf(const FSoftObjectPath& Subobject);
+	
+	/** @return Gets the last object name in the subpath. */
+	CONCERTSHAREDSLATE_API FString ExtractObjectDisplayStringFromPath(const FSoftObjectPath& Object);
+};
diff --git a/Engine/Plugins/Developer/Concert/ConcertUI/ConcertSharedSlate/Source/ConcertSharedSlate/Public/Replication/Utils/ReplicationWidgetDelegates.h b/Engine/Plugins/Developer/Concert/ConcertUI/ConcertSharedSlate/Source/ConcertSharedSlate/Public/Replication/Utils/ReplicationWidgetDelegates.h
new file mode 100644
--- /dev/null
+++ b/Engine/Plugins/Developer/Concert/ConcertUI/ConcertSharedSlate/Source/ConcertSharedSlate/Public/Replication/Utils/ReplicationWidgetDelegates.h
@@ -0,0 +1,29 @@
+// Copyright Epic Games, Inc. All Rights Reserved.
+
+#pragma once
+
+#include "Containers/ContainersFwd.h"
+#include "Delegates/Delegate.h"
+#include "Delegates/DelegateCombinations.h"
+#include "Templates/Function.h"
+
+class FMenuBuilder;
+enum class EBreakBehavior : uint8;
+struct FConcertPropertyChain;
+struct FSoftClassPath;
+struct FSoftObjectPath;
+
+namespace UE::ConcertSharedSlate
+{
+	class FPropertyData;
+	class IPropertySelectionSourceModel;
+	class IReplicationStreamModel;
+
+	/** A predicate for determining Left < Right. */
+	DECLARE_DELEGATE_RetVal_TwoParams(bool, FSortPropertyPredicate, const FPropertyData& Left, const FPropertyData& Right);
+
+	/** Extends a context menu that is being built for a selection of objects. */
+	DECLARE_DELEGATE_TwoParams(FExtendObjectMenu, FMenuBuilder&, TConstArrayView<FSoftObjectPath> ContextObjects);
+	
+	DECLARE_DELEGATE_RetVal_OneParam(FSoftClassPath, FGetObjectClass, const FSoftObjectPath&);
+}
\ No newline at end of file
