Add "# Properties" column to replication object view and align the columns to all be fill width, except for the assignment columns at the end, which is fill sized.

diff --git a/Engine/Plugins/Developer/Concert/ConcertApp/MultiUserClient/Source/MultiUserClient/Private/MultiUserReplicationStyle.cpp b/Engine/Plugins/Developer/Concert/ConcertApp/MultiUserClient/Source/MultiUserClient/Private/MultiUserReplicationStyle.cpp
--- a/Engine/Plugins/Developer/Concert/ConcertApp/MultiUserClient/Source/MultiUserClient/Private/MultiUserReplicationStyle.cpp
+++ b/Engine/Plugins/Developer/Concert/ConcertApp/MultiUserClient/Source/MultiUserClient/Private/MultiUserReplicationStyle.cpp
@@ -39,12 +39,10 @@ namespace UE::MultiUserClient
 		StyleSet->SetContentRoot(FPaths::EngineContentDir() / TEXT("Editor/Slate"));
 		StyleSet->SetCoreContentRoot(FPaths::EngineContentDir() / TEXT("Slate"));
 
-		// Column widths
-		StyleSet->Set("AllClients.Object.OwnerColumnWidth", 200.f);
+		// Column widths - see FConcertFrontendStyle::Initialize() also ("Default Concert Replication Columns Widths")
 		StyleSet->Set("AllClients.Object.ReplicationToggle", 45.f);
-		StyleSet->Set("AllClients.Property.OwnerColumnWidth", 200.f);
-		StyleSet->Set("SingleClient.Object.OwnerColumnWidth", 200.f);
-		StyleSet->Set("SingleClient.Property.OwnerColumnWidth", 200.f);
+		StyleSet->Set("AllClients.Object.OwnerSize", 200.f);
+		StyleSet->Set("AllClients.Property.OwnerSize", 200.f);
 
 		// Timing
 		StyleSet->Set("AllClients.Reassignment.DisplayThrobberAfterSeconds", 0.2f);
diff --git a/Engine/Plugins/Developer/Concert/ConcertApp/MultiUserClient/Source/MultiUserClient/Private/Widgets/ActiveSession/Replication/Client/Multi/Columns/MultiStreamColumns_AssignProperties.cpp b/Engine/Plugins/Developer/Concert/ConcertApp/MultiUserClient/Source/MultiUserClient/Private/Widgets/ActiveSession/Replication/Client/Multi/Columns/MultiStreamColumns_AssignProperties.cpp
--- a/Engine/Plugins/Developer/Concert/ConcertApp/MultiUserClient/Source/MultiUserClient/Private/Widgets/ActiveSession/Replication/Client/Multi/Columns/MultiStreamColumns_AssignProperties.cpp
+++ b/Engine/Plugins/Developer/Concert/ConcertApp/MultiUserClient/Source/MultiUserClient/Private/Widgets/ActiveSession/Replication/Client/Multi/Columns/MultiStreamColumns_AssignProperties.cpp
@@ -101,7 +101,7 @@ namespace UE::MultiUserClient::MultiStreamColumns
 				return SHeaderRow::Column(AssignPropertyColumnId)
 					.DefaultLabel(LOCTEXT("Owner.Label", "Assigned Client"))
 					.ToolTipText(LOCTEXT("Owner.ToolTip", "Client that should replicate this property"))
-					.FillSized(FMultiUserReplicationStyle::Get()->GetFloat(TEXT("AllClients.Property.OwnerColumnWidth")));
+					.FillSized(FMultiUserReplicationStyle::Get()->GetFloat(TEXT("AllClients.Property.OwnerSize")));
 			}
 			
 			virtual TSharedRef<SWidget> GenerateColumnWidget(const FBuildArgs& InArgs) override
diff --git a/Engine/Plugins/Developer/Concert/ConcertApp/MultiUserClient/Source/MultiUserClient/Private/Widgets/ActiveSession/Replication/Client/Multi/Columns/MultiStreamColumns_ReassignOwnership.cpp b/Engine/Plugins/Developer/Concert/ConcertApp/MultiUserClient/Source/MultiUserClient/Private/Widgets/ActiveSession/Replication/Client/Multi/Columns/MultiStreamColumns_ReassignOwnership.cpp
--- a/Engine/Plugins/Developer/Concert/ConcertApp/MultiUserClient/Source/MultiUserClient/Private/Widgets/ActiveSession/Replication/Client/Multi/Columns/MultiStreamColumns_ReassignOwnership.cpp
+++ b/Engine/Plugins/Developer/Concert/ConcertApp/MultiUserClient/Source/MultiUserClient/Private/Widgets/ActiveSession/Replication/Client/Multi/Columns/MultiStreamColumns_ReassignOwnership.cpp
@@ -46,7 +46,7 @@ namespace UE::MultiUserClient::MultiStreamColumns
 				return SHeaderRow::Column(ReassignOwnershipColumnId)
 					.DefaultLabel(LOCTEXT("Owner.Label", "Assigned Clients"))
 					.ToolTipText(LOCTEXT("Owner.ToolTip", "Clients that have registered properties for an object"))
-					.FillSized(FMultiUserReplicationStyle::Get()->GetFloat(TEXT("AllClients.Object.OwnerColumnWidth")));
+					.FillSized(FMultiUserReplicationStyle::Get()->GetFloat(TEXT("AllClients.Object.OwnerSize")));
 			}
 			
 			virtual TSharedRef<SWidget> GenerateColumnWidget(const FBuildArgs& InArgs) override
diff --git a/Engine/Plugins/Developer/Concert/ConcertApp/MultiUserClient/Source/MultiUserClient/Private/Widgets/ActiveSession/Replication/Client/Multi/SMultiClientView.cpp b/Engine/Plugins/Developer/Concert/ConcertApp/MultiUserClient/Source/MultiUserClient/Private/Widgets/ActiveSession/Replication/Client/Multi/SMultiClientView.cpp
--- a/Engine/Plugins/Developer/Concert/ConcertApp/MultiUserClient/Source/MultiUserClient/Private/Widgets/ActiveSession/Replication/Client/Multi/SMultiClientView.cpp
+++ b/Engine/Plugins/Developer/Concert/ConcertApp/MultiUserClient/Source/MultiUserClient/Private/Widgets/ActiveSession/Replication/Client/Multi/SMultiClientView.cpp
@@ -85,7 +85,6 @@ namespace UE::MultiUserClient
 			.AdditionalPropertyColumns =
 			{
 				ReplicationColumns::Property::LabelColumn(),
-				ReplicationColumns::Property::TypeColumn(),
 				MultiStreamColumns::AssignPropertyColumn(MultiStreamEditorAttribute, InConcertClient, InClientManager)
 			}
 		};
diff --git a/Engine/Plugins/Developer/Concert/ConcertUI/ConcertClientSharedSlate/Source/ConcertClientSharedSlate/Private/Replication/Editor/Model/TransactionalReplicationStreamModel.h b/Engine/Plugins/Developer/Concert/ConcertUI/ConcertClientSharedSlate/Source/ConcertClientSharedSlate/Private/Replication/Editor/Model/TransactionalReplicationStreamModel.h
--- a/Engine/Plugins/Developer/Concert/ConcertUI/ConcertClientSharedSlate/Source/ConcertClientSharedSlate/Private/Replication/Editor/Model/TransactionalReplicationStreamModel.h
+++ b/Engine/Plugins/Developer/Concert/ConcertUI/ConcertClientSharedSlate/Source/ConcertClientSharedSlate/Private/Replication/Editor/Model/TransactionalReplicationStreamModel.h
@@ -31,6 +31,7 @@ namespace UE::ConcertSharedSlate
 		virtual bool ContainsProperties(const FSoftObjectPath& Object, const TSet<FConcertPropertyChain>& Properties) const override { return WrappedModel->ContainsProperties(Object, Properties); }
 		virtual bool ForEachReplicatedObject(TFunctionRef<EBreakBehavior(const FSoftObjectPath& Object)> Delegate) const override { return WrappedModel->ForEachReplicatedObject(Delegate); }
 		virtual bool ForEachProperty(const FSoftObjectPath& Object, TFunctionRef<EBreakBehavior(const FConcertPropertyChain& Property)> Delegate) const override { return WrappedModel->ForEachProperty(Object, Delegate); }
+		virtual uint32 GetNumProperties(const FSoftObjectPath& Object) const override { return WrappedModel->GetNumProperties(Object); }
 		virtual FOnObjectsChanged& OnObjectsChanged() override { return WrappedModel->OnObjectsChanged(); }
 		virtual FOnPropertiesChanged& OnPropertiesChanged() override { return WrappedModel->OnPropertiesChanged(); }
 		//~ End IReplicationStreamModel Interface
diff --git a/Engine/Plugins/Developer/Concert/ConcertUI/ConcertClientSharedSlate/Source/ConcertClientSharedSlate/Public/Replication/ClientReplicationWidgetFactories.h b/Engine/Plugins/Developer/Concert/ConcertUI/ConcertClientSharedSlate/Source/ConcertClientSharedSlate/Public/Replication/ClientReplicationWidgetFactories.h
--- a/Engine/Plugins/Developer/Concert/ConcertUI/ConcertClientSharedSlate/Source/ConcertClientSharedSlate/Public/Replication/ClientReplicationWidgetFactories.h
+++ b/Engine/Plugins/Developer/Concert/ConcertUI/ConcertClientSharedSlate/Source/ConcertClientSharedSlate/Public/Replication/ClientReplicationWidgetFactories.h
@@ -50,8 +50,7 @@ namespace UE::ConcertClientSharedSlate
 		/** The columns the property view should have. The label column is always included. */
 		TArray<ConcertSharedSlate::FPropertyColumnEntry> AdditionalPropertyColumns
 		{
-			ConcertSharedSlate::ReplicationColumns::Property::LabelColumn(),
-			ConcertSharedSlate::ReplicationColumns::Property::TypeColumn()
+			ConcertSharedSlate::ReplicationColumns::Property::LabelColumn()
 		};
 		
 		/** Optional initial primary sort mode for object rows */
diff --git a/Engine/Plugins/Developer/Concert/ConcertUI/ConcertSharedSlate/Source/ConcertSharedSlate/Private/ConcertFrontendStyle.cpp b/Engine/Plugins/Developer/Concert/ConcertUI/ConcertSharedSlate/Source/ConcertSharedSlate/Private/ConcertFrontendStyle.cpp
--- a/Engine/Plugins/Developer/Concert/ConcertUI/ConcertSharedSlate/Source/ConcertSharedSlate/Private/ConcertFrontendStyle.cpp
+++ b/Engine/Plugins/Developer/Concert/ConcertUI/ConcertSharedSlate/Source/ConcertSharedSlate/Private/ConcertFrontendStyle.cpp
@@ -148,10 +148,13 @@ void FConcertFrontendStyle::Initialize()
 	// Colors
 	StyleSet->Set("Concert.DisconnectedColor", FLinearColor(0.672f, 0.672f, 0.672f));
 
-	// Floats
+	// Default Concert Replication Columns Widths
 	StyleSet->Set("Concert.Replication.Tree.RowHeight", 22.f);
-	StyleSet->Set("Concert.Replication.Tree.Object.LabelRowWidth", 255.f);
-	StyleSet->Set("Concert.Replication.Tree.Property.LabelRowWidth", 300.f);
+	StyleSet->Set("Concert.Replication.Object.LabelSize", 1.5f);
+	StyleSet->Set("Concert.Replication.Object.TypeWidth", 1.f);
+	StyleSet->Set("Concert.Replication.Object.NumPropertiesSize", 1.f);
+	StyleSet->Set("Concert.Replication.Property.LabelSize", 1.35f); // Property label needs to be smaller to align the top and bottom type columns
+	StyleSet->Set("Concert.Replication.Property.TypeSize", 1.f);
 	
 	FSlateStyleRegistry::RegisterSlateStyle(*StyleSet.Get());
 };
diff --git a/Engine/Plugins/Developer/Concert/ConcertUI/ConcertSharedSlate/Source/ConcertSharedSlate/Private/Replication/Editor/Model/GenericReplicationStreamModel.cpp b/Engine/Plugins/Developer/Concert/ConcertUI/ConcertSharedSlate/Source/ConcertSharedSlate/Private/Replication/Editor/Model/GenericReplicationStreamModel.cpp
--- a/Engine/Plugins/Developer/Concert/ConcertUI/ConcertSharedSlate/Source/ConcertSharedSlate/Private/Replication/Editor/Model/GenericReplicationStreamModel.cpp
+++ b/Engine/Plugins/Developer/Concert/ConcertUI/ConcertSharedSlate/Source/ConcertSharedSlate/Private/Replication/Editor/Model/GenericReplicationStreamModel.cpp
@@ -151,6 +151,18 @@ namespace UE::ConcertSharedSlate
 		return !AssignedProperties->PropertySelection.ReplicatedProperties.IsEmpty();
 	}
 
+	uint32 FGenericReplicationStreamModel::GetNumProperties(const FSoftObjectPath& Object) const
+	{
+		const FConcertObjectReplicationMap* ReplicationMap = ReplicationMapAttribute.Get();
+		if (!ensure(ReplicationMap))
+		{
+			return 0;
+		}
+		
+		const FConcertReplicatedObjectInfo* AssignedProperties = ReplicationMap->ReplicatedObjects.Find(Object);
+		return AssignedProperties ? AssignedProperties->PropertySelection.ReplicatedProperties.Num() : 0;
+	}
+
 	void FGenericReplicationStreamModel::AddObjects(TConstArrayView<UObject*> Objects)
 	{
 		FConcertObjectReplicationMap* ReplicationMap = ReplicationMapAttribute.Get();
diff --git a/Engine/Plugins/Developer/Concert/ConcertUI/ConcertSharedSlate/Source/ConcertSharedSlate/Private/Replication/Editor/Model/GenericReplicationStreamModel.h b/Engine/Plugins/Developer/Concert/ConcertUI/ConcertSharedSlate/Source/ConcertSharedSlate/Private/Replication/Editor/Model/GenericReplicationStreamModel.h
--- a/Engine/Plugins/Developer/Concert/ConcertUI/ConcertSharedSlate/Source/ConcertSharedSlate/Private/Replication/Editor/Model/GenericReplicationStreamModel.h
+++ b/Engine/Plugins/Developer/Concert/ConcertUI/ConcertSharedSlate/Source/ConcertSharedSlate/Private/Replication/Editor/Model/GenericReplicationStreamModel.h
@@ -30,6 +30,7 @@ namespace UE::ConcertSharedSlate
 		virtual bool ContainsProperties(const FSoftObjectPath& Object, const TSet<FConcertPropertyChain>& Properties) const override;
 		virtual bool ForEachReplicatedObject(TFunctionRef<EBreakBehavior(const FSoftObjectPath& Object)> Delegate) const override;
 		virtual bool ForEachProperty(const FSoftObjectPath& Object, TFunctionRef<EBreakBehavior(const FConcertPropertyChain& Parent)> Delegate) const override;
+		virtual uint32 GetNumProperties(const FSoftObjectPath& Object) const override;
 		//~ End IReplicationStreamModel Interface
 		
 		//~ Begin IEditableReplicationStreamModel Interface
diff --git a/Engine/Plugins/Developer/Concert/ConcertUI/ConcertSharedSlate/Source/ConcertSharedSlate/Private/Replication/Editor/View/MultiEditor/ConsolidatedMultiStreamModel.cpp b/Engine/Plugins/Developer/Concert/ConcertUI/ConcertSharedSlate/Source/ConcertSharedSlate/Private/Replication/Editor/View/MultiEditor/ConsolidatedMultiStreamModel.cpp
--- a/Engine/Plugins/Developer/Concert/ConcertUI/ConcertSharedSlate/Source/ConcertSharedSlate/Private/Replication/Editor/View/MultiEditor/ConsolidatedMultiStreamModel.cpp
+++ b/Engine/Plugins/Developer/Concert/ConcertUI/ConcertSharedSlate/Source/ConcertSharedSlate/Private/Replication/Editor/View/MultiEditor/ConsolidatedMultiStreamModel.cpp
@@ -62,7 +62,18 @@ namespace UE::ConcertSharedSlate
 			return bContains ? EBreakBehavior::Break : EBreakBehavior::Continue;
 		});
 		return bContains;
-	}		
+	}
+
+	bool FConsolidatedMultiStreamModel::ContainsProperties(const FSoftObjectPath& Object, const TSet<FConcertPropertyChain>& Properties) const
+	{
+		bool bContainsProperty = false;
+		MultiStreamModel->ForEachStream([&Object, &Properties, &bContainsProperty](const TSharedRef<IReplicationStreamModel>& Model)
+		{
+			bContainsProperty = Model->ContainsProperties(Object, Properties);
+			return bContainsProperty ? EBreakBehavior::Break : EBreakBehavior::Continue;
+		});
+		return bContainsProperty;
+	}
 
 	bool FConsolidatedMultiStreamModel::ForEachReplicatedObject(TFunctionRef<EBreakBehavior(const FSoftObjectPath& Object)> Delegate) const
 	{
@@ -95,6 +106,33 @@ namespace UE::ConcertSharedSlate
 		return bAnyMappings;
 	}
 
+	bool FConsolidatedMultiStreamModel::ForEachProperty(const FSoftObjectPath& Object, TFunctionRef<EBreakBehavior(const FConcertPropertyChain& Property)> Delegate) const
+	{
+		bool bCalledAtLeastOnce = false;
+		MultiStreamModel->ForEachStream([&Object, &Delegate, &bCalledAtLeastOnce](const TSharedRef<IReplicationStreamModel>& Model)
+		{
+			EBreakBehavior BreakBehavior = EBreakBehavior::Continue;
+			bCalledAtLeastOnce |= Model->ForEachProperty(Object, [&Delegate, &BreakBehavior](const FConcertPropertyChain& Property)
+			{
+				BreakBehavior = Delegate(Property);
+				return BreakBehavior;
+			});
+			return BreakBehavior;
+		});
+		return bCalledAtLeastOnce;
+	}
+
+	uint32 FConsolidatedMultiStreamModel::GetNumProperties(const FSoftObjectPath& Object) const
+	{
+		uint32 NumProperties = 0;
+		MultiStreamModel->ForEachStream([&Object, &NumProperties](const TSharedRef<IReplicationStreamModel>& Model)
+		{
+			NumProperties += Model->GetNumProperties(Object);
+			return EBreakBehavior::Continue;
+		});
+		return NumProperties;
+	}
+
 	void FConsolidatedMultiStreamModel::AddObjects(TConstArrayView<UObject*> Objects)
 	{
 		const TSharedPtr<IEditableReplicationStreamModel> TargetStream = GetAutoAssignTargetDelegate.IsBound()
diff --git a/Engine/Plugins/Developer/Concert/ConcertUI/ConcertSharedSlate/Source/ConcertSharedSlate/Private/Replication/Editor/View/MultiEditor/ConsolidatedMultiStreamModel.h b/Engine/Plugins/Developer/Concert/ConcertUI/ConcertSharedSlate/Source/ConcertSharedSlate/Private/Replication/Editor/View/MultiEditor/ConsolidatedMultiStreamModel.h
--- a/Engine/Plugins/Developer/Concert/ConcertUI/ConcertSharedSlate/Source/ConcertSharedSlate/Private/Replication/Editor/View/MultiEditor/ConsolidatedMultiStreamModel.h
+++ b/Engine/Plugins/Developer/Concert/ConcertUI/ConcertSharedSlate/Source/ConcertSharedSlate/Private/Replication/Editor/View/MultiEditor/ConsolidatedMultiStreamModel.h
@@ -38,21 +38,20 @@ namespace UE::ConcertSharedSlate
 			);
 		virtual ~FConsolidatedMultiStreamModel();
 		
-		// The functions related to properties are checkNoEntry() because SBaseReplicationStreamEditor's FFakeObjectToPropertiesEditorModel masks those functions 
-		
 		//~ Begin IReplicationStreamModel Interface
 		virtual FSoftClassPath GetObjectClass(const FSoftObjectPath& Object) const override;
 		virtual bool ContainsObjects(const TSet<FSoftObjectPath>& Objects) const override;
-		virtual bool ContainsProperties(const FSoftObjectPath& Object, const TSet<FConcertPropertyChain>& Properties) const override { checkNoEntry(); return false; }
+		virtual bool ContainsProperties(const FSoftObjectPath& Object, const TSet<FConcertPropertyChain>& Properties) const override;
 		virtual bool ForEachReplicatedObject(TFunctionRef<EBreakBehavior(const FSoftObjectPath& Object)> Delegate) const override;
-		virtual bool ForEachProperty(const FSoftObjectPath& Object, TFunctionRef<EBreakBehavior(const FConcertPropertyChain& Property)> Delegate) const override { checkNoEntry(); return false; }
+		virtual bool ForEachProperty(const FSoftObjectPath& Object, TFunctionRef<EBreakBehavior(const FConcertPropertyChain& Property)> Delegate) const override;
+		virtual uint32 GetNumProperties(const FSoftObjectPath& Object) const override;
 		//~ End IReplicationStreamModel Interface
 
 		//~ Begin IEditableReplicationStreamModel Interface
 		virtual void AddObjects(TConstArrayView<UObject*> Objects) override;
 		virtual void RemoveObjects(TConstArrayView<FSoftObjectPath> Objects) override;
-		virtual void AddProperties(const FSoftObjectPath& Object, TConstArrayView<FConcertPropertyChain> Properties) override { checkNoEntry(); }
-		virtual void RemoveProperties(const FSoftObjectPath& Object, TConstArrayView<FConcertPropertyChain> Properties) override { checkNoEntry(); }
+		virtual void AddProperties(const FSoftObjectPath& Object, TConstArrayView<FConcertPropertyChain> Properties) override { ensure(false); }
+		virtual void RemoveProperties(const FSoftObjectPath& Object, TConstArrayView<FConcertPropertyChain> Properties) override { ensure(false); }
 		virtual FOnObjectsChanged& OnObjectsChanged() override { return OnObjectsChangedDelegate; }
 		virtual FOnPropertiesChanged& OnPropertiesChanged() override { return OnPropertiesChangedDelegate; }
 		//~ End IEditableReplicationStreamModel Interface
diff --git a/Engine/Plugins/Developer/Concert/ConcertUI/ConcertSharedSlate/Source/ConcertSharedSlate/Private/Replication/Editor/View/ObjectViewer/SReplicationStreamViewer.cpp b/Engine/Plugins/Developer/Concert/ConcertUI/ConcertSharedSlate/Source/ConcertSharedSlate/Private/Replication/Editor/View/ObjectViewer/SReplicationStreamViewer.cpp
--- a/Engine/Plugins/Developer/Concert/ConcertUI/ConcertSharedSlate/Source/ConcertSharedSlate/Private/Replication/Editor/View/ObjectViewer/SReplicationStreamViewer.cpp
+++ b/Engine/Plugins/Developer/Concert/ConcertUI/ConcertSharedSlate/Source/ConcertSharedSlate/Private/Replication/Editor/View/ObjectViewer/SReplicationStreamViewer.cpp
@@ -211,6 +211,7 @@ namespace UE::ConcertSharedSlate
 		TArray<FObjectColumnEntry> Columns = InArgs._ObjectColumns;
 		Columns.Add(ReplicationColumns::TopLevel::LabelColumn(NameModel.Get(), GetObjectClassDelegate));
 		Columns.Add(ReplicationColumns::TopLevel::TypeColumn(GetObjectClassDelegate));
+		Columns.Add(ReplicationColumns::TopLevel::NumPropertiesColumn(*PropertiesModel));
 		
 		const bool bHasNoOutlinerObjectsAttribute = InArgs._NoOutlinerObjects.IsBound() || InArgs._NoOutlinerObjects.IsSet(); 
 		const TAttribute<FText> NoObjectsAttribute = bHasNoOutlinerObjectsAttribute ? InArgs._NoOutlinerObjects : LOCTEXT("NoObjects", "No objects to display");
diff --git a/Engine/Plugins/Developer/Concert/ConcertUI/ConcertSharedSlate/Source/ConcertSharedSlate/Private/Replication/Editor/View/SelectionViewerColumns.cpp b/Engine/Plugins/Developer/Concert/ConcertUI/ConcertSharedSlate/Source/ConcertSharedSlate/Private/Replication/Editor/View/SelectionViewerColumns.cpp
--- a/Engine/Plugins/Developer/Concert/ConcertUI/ConcertSharedSlate/Source/ConcertSharedSlate/Private/Replication/Editor/View/SelectionViewerColumns.cpp
+++ b/Engine/Plugins/Developer/Concert/ConcertUI/ConcertSharedSlate/Source/ConcertSharedSlate/Private/Replication/Editor/View/SelectionViewerColumns.cpp
@@ -3,15 +3,15 @@
 #include "Replication/Editor/View/Column/SelectionViewerColumns.h"
 
 #include "ConcertFrontendStyle.h"
-#include "Replication/Editor/Model/IEditableReplicationStreamModel.h"
 #include "Replication/Editor/Model/Data/PropertyData.h"
 #include "Replication/Editor/Model/Data/ReplicatedObjectData.h"
+#include "Replication/Editor/Model/IReplicationStreamModel.h"
 #include "Replication/Editor/Utils/DisplayUtils.h"
+#include "Replication/Editor/View/Column/IObjectTreeColumn.h"
 #include "Replication/Editor/View/Column/ReplicationColumnsUtils.h"
 #include "Replication/PropertyChainUtils.h"
 
 #include "Internationalization/Internationalization.h"
-#include "Replication/Editor/View/Column/IObjectTreeColumn.h"
 #include "Textures/SlateIcon.h"
 #include "Widgets/Images/SImage.h"
 #include "Widgets/Input/SCheckBox.h"
@@ -24,9 +24,9 @@
 
 namespace UE::ConcertSharedSlate::ReplicationColumns::TopLevel
 {
-	const FName IconColumnId = TEXT("IconColumn");
 	const FName LabelColumnId = TEXT("LabelColumn");
 	const FName TypeColumnId = TEXT("TypeColumn");
+	const FName NumPropertiesColumnId = TEXT("NumPropertiesColumnId");
 	
 	FObjectColumnEntry LabelColumn(IObjectNameModel* OptionalNameModel, FGetObjectClass GetObjectClassDelegate)
 	{
@@ -43,7 +43,7 @@ namespace UE::ConcertSharedSlate::ReplicationColumns::TopLevel
 			{
 				return SHeaderRow::Column(LabelColumnId)
 					.DefaultLabel(LOCTEXT("LabelColumnLabel", "Label"))
-					.FillSized(FConcertFrontendStyle::Get()->GetFloat("Concert.Replication.Tree.Object.LabelRowWidth"));
+					.FillWidth(FConcertFrontendStyle::Get()->GetFloat("Concert.Replication.Object.LabelSize"));
 			}
 			
 			virtual TSharedRef<SWidget> GenerateColumnWidget(const FBuildArgs& InArgs) override
@@ -123,8 +123,8 @@ namespace UE::ConcertSharedSlate::ReplicationColumns::TopLevel
 			virtual SHeaderRow::FColumn::FArguments CreateHeaderRowArgs() const override
 			{
 				return SHeaderRow::Column(TypeColumnId)
-					.DefaultLabel(LOCTEXT("TypeColumnLabel", "Type"))
-					.FillWidth(1.f);
+					.DefaultLabel(LOCTEXT("TypeColumnLabel", "Type")) 
+					.FillWidth(FConcertFrontendStyle::Get()->GetFloat("Concert.Replication.Object.TypeWidth"));
 			}
 			
 			virtual TSharedRef<SWidget> GenerateColumnWidget(const FBuildArgs& InArgs) override
@@ -136,6 +136,7 @@ namespace UE::ConcertSharedSlate::ReplicationColumns::TopLevel
 						SNew(STextBlock)
 						.HighlightText(TAttribute<FText>::CreateLambda([HighlightText = InArgs.HighlightText](){ return *HighlightText; }))
 						.Text(DisplayUtils::GetObjectTypeText(Class))
+						.ColorAndOpacity(FSlateColor::UseSubduedForeground())
 					];
 			}
 			
@@ -171,6 +172,94 @@ namespace UE::ConcertSharedSlate::ReplicationColumns::TopLevel
 			{ static_cast<int32>(ETopLevelColumnOrder::Type) }
 		};
 	}
+
+	FObjectColumnEntry NumPropertiesColumn(const IReplicationStreamModel& Model, ENumPropertiesFlags Flags)
+	{
+		class FNumPropertiesColumn_Object : public IObjectTreeColumn
+		{
+		public:
+			
+			FNumPropertiesColumn_Object(const IReplicationStreamModel& Model UE_LIFETIMEBOUND, ENumPropertiesFlags Flags)
+				: Model(Model)
+				, Flags(Flags)
+			{}
+
+			virtual SHeaderRow::FColumn::FArguments CreateHeaderRowArgs() const override
+			{
+				return SHeaderRow::Column(NumPropertiesColumnId)
+					.DefaultLabel(LOCTEXT("NumPropertyColumnLabel", "# Properties"))
+					.FillWidth(FConcertFrontendStyle::Get()->GetFloat("Concert.Replication.Object.NumPropertiesSize"));
+			}
+			
+			virtual TSharedRef<SWidget> GenerateColumnWidget(const FBuildArgs& InArgs) override
+			{
+				return SNew(SHorizontalBox)
+					
+					+SHorizontalBox::Slot()
+					.HAlign(HAlign_Left)
+					.VAlign(VAlign_Center)
+					.Padding(6.f, 0.f, 0.f, 0.f)
+					[
+						SNew(STextBlock)
+						.HighlightText(TAttribute<FText>::CreateLambda([HighlightText = InArgs.HighlightText](){ return *HighlightText; }))
+						.Text_Raw(this, &FNumPropertiesColumn_Object::GetDisplayText, InArgs.RowItem.RowData.GetObjectPath())
+						.ColorAndOpacity(FSlateColor::UseSubduedForeground())
+					];
+			}
+			
+			virtual void PopulateSearchString(const FObjectTreeRowContext& InItem, TArray<FString>& InOutSearchStrings) const override
+			{
+				InOutSearchStrings.Add(
+					FString::FromInt(GetNumProperties(InItem.RowData.GetObjectPath()))
+					);
+			}
+			
+			virtual bool CanBeSorted() const override { return true; } 
+			virtual bool IsLessThan(const FObjectTreeRowContext& Left, const FObjectTreeRowContext& Right) const override
+			{
+				return GetNumProperties(Left.RowData.GetObjectPath()) < GetNumProperties(Right.RowData.GetObjectPath());
+			}
+
+		private:
+			
+			const IReplicationStreamModel& Model;
+			const ENumPropertiesFlags Flags;
+			
+			FText GetDisplayText(FSoftObjectPath ObjectPath) const
+			{
+				return FText::AsNumber(GetNumProperties(ObjectPath));
+			}
+
+			uint32 GetNumProperties(const FSoftObjectPath& ObjectPath) const
+			{
+				const uint32 ObjectProperties = Model.GetNumProperties(ObjectPath);
+				
+				const bool bCountSubobjects = EnumHasAnyFlags(Flags, ENumPropertiesFlags::IncludeSubobjectCounts);
+				if (!bCountSubobjects)
+				{
+					return ObjectProperties;
+				}
+
+				uint32 NumSubobjectProperties = 0;
+				Model.ForEachSubobject(ObjectPath, [this, &NumSubobjectProperties](const FSoftObjectPath& Child)
+				{
+					NumSubobjectProperties += Model.GetNumProperties(Child);
+					return EBreakBehavior::Continue;
+				});
+
+				return ObjectProperties + NumSubobjectProperties;
+			}
+		};
+		
+		return {
+			TReplicationColumnDelegates<FObjectTreeRowContext>::FCreateColumn::CreateLambda([&Model, Flags]()
+			{
+				return MakeShared<FNumPropertiesColumn_Object>(Model, Flags);
+			}),
+			NumPropertiesColumnId,
+			{ static_cast<int32>(ETopLevelColumnOrder::NumProperties) }
+		};
+	}
 }
 
 #undef LOCTEXT_NAMESPACE
@@ -192,7 +281,7 @@ namespace UE::ConcertSharedSlate::ReplicationColumns::Property
 			{
 				return SHeaderRow::Column(LabelColumnId)
 					.DefaultLabel(LOCTEXT("LabelColumnLabel", "Label"))
-					.FillSized(FConcertFrontendStyle::Get()->GetFloat("Concert.Replication.Tree.Property.LabelRowWidth"));
+					.FillWidth(FConcertFrontendStyle::Get()->GetFloat("Concert.Replication.Property.LabelSize"));
 			}
 			
 			virtual TSharedRef<SWidget> GenerateColumnWidget(const FBuildArgs& InArgs) override
@@ -252,7 +341,7 @@ namespace UE::ConcertSharedSlate::ReplicationColumns::Property
 			{
 				return SHeaderRow::Column(TypeColumnId)
 					.DefaultLabel(LOCTEXT("TypeColumnLabel", "Type"))
-					.FillWidth(1.f);
+					.FillWidth(FConcertFrontendStyle::Get()->GetFloat("Concert.Replication.Property.TypeSize"));
 			}
 			
 			virtual TSharedRef<SWidget> GenerateColumnWidget(const FBuildArgs& InArgs) override
@@ -265,6 +354,7 @@ namespace UE::ConcertSharedSlate::ReplicationColumns::Property
 						.Font(FAppStyle::GetFontStyle("PropertyWindow.NormalFont"))
 						.HighlightText(TAttribute<FText>::CreateLambda([HighlightText = InArgs.HighlightText](){ return *HighlightText; }))
 						.Text(GetDisplayText(InArgs.RowItem.RowData))
+						.ColorAndOpacity(FSlateColor::UseSubduedForeground())
 					];
 			}
 			
diff --git a/Engine/Plugins/Developer/Concert/ConcertUI/ConcertSharedSlate/Source/ConcertSharedSlate/Public/Replication/Editor/Model/IReplicationStreamModel.h b/Engine/Plugins/Developer/Concert/ConcertUI/ConcertSharedSlate/Source/ConcertSharedSlate/Public/Replication/Editor/Model/IReplicationStreamModel.h
--- a/Engine/Plugins/Developer/Concert/ConcertUI/ConcertSharedSlate/Source/ConcertSharedSlate/Public/Replication/Editor/Model/IReplicationStreamModel.h
+++ b/Engine/Plugins/Developer/Concert/ConcertUI/ConcertSharedSlate/Source/ConcertSharedSlate/Public/Replication/Editor/Model/IReplicationStreamModel.h
@@ -42,6 +42,9 @@ namespace UE::ConcertSharedSlate
 		 */
 		virtual bool ForEachProperty(const FSoftObjectPath& Object, TFunctionRef<EBreakBehavior(const FConcertPropertyChain& Property)> Delegate) const = 0;
 
+		/** Gets the number of properties assigned to Object. */
+		virtual uint32 GetNumProperties(const FSoftObjectPath& Object) const = 0;
+
 		/** @return Gets all replicated objects */
 		TSet<FSoftObjectPath> GetReplicatedObjects() const
 		{
@@ -105,6 +108,13 @@ namespace UE::ConcertSharedSlate
 		TSet<FConcertPropertyChain> GetAllProperties(const FSoftObjectPath& Object) const
 		{
 			TSet<FConcertPropertyChain> Result;
+			const uint32 NumProperties = GetNumProperties(Object);
+			if (NumProperties == 0)
+			{
+				return Result;
+			}
+
+			Result.Reserve(NumProperties);
 			ForEachProperty(Object, [&Result](const FConcertPropertyChain& Property)
 			{
 				Result.Add(Property);
diff --git a/Engine/Plugins/Developer/Concert/ConcertUI/ConcertSharedSlate/Source/ConcertSharedSlate/Public/Replication/Editor/View/Column/SelectionViewerColumns.h b/Engine/Plugins/Developer/Concert/ConcertUI/ConcertSharedSlate/Source/ConcertSharedSlate/Public/Replication/Editor/View/Column/SelectionViewerColumns.h
--- a/Engine/Plugins/Developer/Concert/ConcertUI/ConcertSharedSlate/Source/ConcertSharedSlate/Public/Replication/Editor/View/Column/SelectionViewerColumns.h
+++ b/Engine/Plugins/Developer/Concert/ConcertUI/ConcertSharedSlate/Source/ConcertSharedSlate/Public/Replication/Editor/View/Column/SelectionViewerColumns.h
@@ -7,6 +7,8 @@
 #include "Replication/Utils/ReplicationWidgetDelegates.h"
 #include "ReplicationColumnInfo.h"
 
+#include "HAL/Platform.h"
+#include "Misc/EnumClassFlags.h"
 #include "Templates/SharedPointer.h"
 
 namespace UE::ConcertSharedSlate
@@ -21,9 +23,9 @@ namespace UE::ConcertSharedSlate
 
 namespace UE::ConcertSharedSlate::ReplicationColumns::TopLevel
 {
-	CONCERTSHAREDSLATE_API extern const FName IconColumnId;
 	CONCERTSHAREDSLATE_API extern const FName LabelColumnId;
 	CONCERTSHAREDSLATE_API extern const FName TypeColumnId;
+	CONCERTSHAREDSLATE_API extern const FName NumPropertiesColumnId;
 
 	enum class ETopLevelColumnOrder : int32
 	{
@@ -31,14 +33,27 @@ namespace UE::ConcertSharedSlate::ReplicationColumns::TopLevel
 		Label = 20,
 		/** Class of the object */
 		Type = 30,
+		/** Displays the number of properties assigned to the object */
+		NumProperties = 40
 	};
 
+	enum class ENumPropertiesFlags : uint8
+	{
+		None,
+		/** Also add up the number of properties that child objects have assigned. */
+		IncludeSubobjectCounts
+	};
+	ENUM_CLASS_FLAGS(ENumPropertiesFlags);
+
 	/**
 	 * @param OptionalNameModel Used to look object name. Defaults to name displayed in object path if unset.
 	 * @param GetObjectClassDelegate Used to display class icon. No icon is displayed if unset.
 	 */
 	CONCERTSHAREDSLATE_API FObjectColumnEntry LabelColumn(IObjectNameModel* OptionalNameModel = nullptr, FGetObjectClass GetObjectClassDelegate = {});
+	/** Shows the object's type. */
 	CONCERTSHAREDSLATE_API FObjectColumnEntry TypeColumn(FGetObjectClass GetObjectClassDelegate);
+	/** Shows the number of widgets assigned to the object. */
+	CONCERTSHAREDSLATE_API FObjectColumnEntry NumPropertiesColumn(const IReplicationStreamModel& Model UE_LIFETIMEBOUND, ENumPropertiesFlags Flags = ENumPropertiesFlags::None);
 }
 
 namespace UE::ConcertSharedSlate::ReplicationColumns::Property
diff --git a/Engine/Plugins/Developer/Concert/ConcertUI/ConcertSharedSlate/Source/ConcertSharedSlate/Public/Replication/PropertyTreeFactory.h b/Engine/Plugins/Developer/Concert/ConcertUI/ConcertSharedSlate/Source/ConcertSharedSlate/Public/Replication/PropertyTreeFactory.h
--- a/Engine/Plugins/Developer/Concert/ConcertUI/ConcertSharedSlate/Source/ConcertSharedSlate/Public/Replication/PropertyTreeFactory.h
+++ b/Engine/Plugins/Developer/Concert/ConcertUI/ConcertSharedSlate/Source/ConcertSharedSlate/Public/Replication/PropertyTreeFactory.h
@@ -23,8 +23,7 @@ namespace UE::ConcertSharedSlate
 		/** Optional. Additional property columns you want added. */
 		TArray<FPropertyColumnEntry> PropertyColumns
 		{
-			ReplicationColumns::Property::LabelColumn(),
-			ReplicationColumns::Property::TypeColumn()
+			ReplicationColumns::Property::LabelColumn()
 		};
 
 		/** Optional filter function. Return true to al */
