Made SClientName, SLocalClientName, SRemoteClientName and SHorizontalClientList generic by moving it to ConcertSharedSlate.
- Unified all locations (browser top-right and overview list) in the Concert UI to use SClientName and related classes.

diff --git a/Engine/Plugins/Developer/Concert/ConcertApp/MultiUserClient/Source/MultiUserClient/Private/Widgets/ActiveSession/Overview/SActiveSessionOverviewTab.cpp b/Engine/Plugins/Developer/Concert/ConcertApp/MultiUserClient/Source/MultiUserClient/Private/Widgets/ActiveSession/Overview/SActiveSessionOverviewTab.cpp
--- a/Engine/Plugins/Developer/Concert/ConcertApp/MultiUserClient/Source/MultiUserClient/Private/Widgets/ActiveSession/Overview/SActiveSessionOverviewTab.cpp
+++ b/Engine/Plugins/Developer/Concert/ConcertApp/MultiUserClient/Source/MultiUserClient/Private/Widgets/ActiveSession/Overview/SActiveSessionOverviewTab.cpp
@@ -13,10 +13,11 @@
 #include "IConcertClient.h"
 #include "IConcertClientWorkspace.h"
 #include "IConcertSyncClient.h"
+#include "Widgets/Client/ClientInfoHelpers.h"
+#include "Widgets/Client/SClientName.h"
 #include "Widgets/ClientSessionHistoryController.h"
 
 #include "Algo/Transform.h"
-#include "EditorFontGlyphs.h"
 #include "FileHelpers.h"
 #include "Misc/AsyncTaskNotification.h"
 #include "Misc/PackageName.h"
@@ -42,10 +43,11 @@ namespace ActiveSessionDetailsUI
 
 class SActiveSessionDetailsRow : public SMultiColumnTableRow<TSharedPtr<FConcertSessionClientInfo>>
 {
+public:
+	
 	SLATE_BEGIN_ARGS(SActiveSessionDetailsRow) {}
 	SLATE_END_ARGS()
 
-public:
 	/**
 	 * Constructs the widget.
 	 *
@@ -65,31 +67,21 @@ public:
 		SetToolTipText(MakeAttributeSP(this, &SActiveSessionDetailsRow::GetRowToolTip));
 	}
 
-public:
 	virtual TSharedRef<SWidget> GenerateWidgetForColumn(const FName& ColumnName) override
 	{
 		if (ColumnName == ActiveSessionDetailsUI::DisplayNameColumnName)
 		{
-			// Displays a colored square from a special font (using avatar color) followed by the the display name -> [x] John Smith
-			return SNew(SHorizontalBox)
-				// The 'square' glyph in front of the client name, rendered using special font glyph, in the client avatar color.
-				+SHorizontalBox::Slot()
-				.AutoWidth()
-				.VAlign(VAlign_Center)
-				.Padding(FMargin(4.0f, 0.0f, 0.0f, 0.0f))
-				[
-					SNew(STextBlock)
-					.Font(this, &SActiveSessionDetailsRow::GetAvatarFont)
-					.ColorAndOpacity(this, &SActiveSessionDetailsRow::GetAvatarColor)
-					.Text(FEditorFontGlyphs::Square)
-				]
+			const TSharedPtr<IConcertSyncClient> SyncClientPin = SyncClient.Pin();
+			const TSharedPtr<FConcertSessionClientInfo> SessionClientInfoPin = SessionClientInfo.Pin();
+			if (!SyncClientPin)
+			{
+				return SNew(STextBlock).Text(LOCTEXT("Unavailable", "Unavailable"));
+			}
 
-				// The client display name.
-				+SHorizontalBox::Slot()
-				.VAlign(VAlign_Center)
-				[
-					ConcertFrontendUtils::CreateDisplayName(MakeAttributeSP(this, &SActiveSessionDetailsRow::GetDisplayName))
-				];
+			const bool bIsLocalClient = SyncClientPin->GetConcertClient()->GetCurrentSession()->GetSessionClientEndpointId() == SessionClientInfoPin->ClientEndpointId;
+			return SNew(UE::ConcertClientSharedSlate::SClientName)
+				.ClientInfo(UE::ConcertClientSharedSlate::MakeClientInfoAttribute(SyncClientPin->GetConcertClient(), SessionClientInfoPin->ClientEndpointId))
+				.DisplayAsLocalClient(bIsLocalClient);
 		}
 		else if (ColumnName == ActiveSessionDetailsUI::PresenceColumnName)
 		{
diff --git a/Engine/Plugins/Developer/Concert/ConcertApp/MultiUserClient/Source/MultiUserClient/Private/Widgets/ActiveSession/Replication/Client/ClientUtils.cpp b/Engine/Plugins/Developer/Concert/ConcertApp/MultiUserClient/Source/MultiUserClient/Private/Widgets/ActiveSession/Replication/Client/ClientUtils.cpp
--- a/Engine/Plugins/Developer/Concert/ConcertApp/MultiUserClient/Source/MultiUserClient/Private/Widgets/ActiveSession/Replication/Client/ClientUtils.cpp
+++ b/Engine/Plugins/Developer/Concert/ConcertApp/MultiUserClient/Source/MultiUserClient/Private/Widgets/ActiveSession/Replication/Client/ClientUtils.cpp
@@ -5,7 +5,7 @@
 #include "ConcertMessageData.h"
 #include "IConcertClient.h"
 #include "Replication/Client/ReplicationClientManager.h"
-#include "Widgets/ClientName/SClientName.h"
+#include "Widgets/Client/SClientName.h"
 
 namespace UE::MultiUserClient::ClientUtils
 {
diff --git a/Engine/Plugins/Developer/Concert/ConcertApp/MultiUserClient/Source/MultiUserClient/Private/Widgets/ActiveSession/Replication/Client/Multi/Columns/SAssignPropertyComboBox.cpp b/Engine/Plugins/Developer/Concert/ConcertApp/MultiUserClient/Source/MultiUserClient/Private/Widgets/ActiveSession/Replication/Client/Multi/Columns/SAssignPropertyComboBox.cpp
--- a/Engine/Plugins/Developer/Concert/ConcertApp/MultiUserClient/Source/MultiUserClient/Private/Widgets/ActiveSession/Replication/Client/Multi/Columns/SAssignPropertyComboBox.cpp
+++ b/Engine/Plugins/Developer/Concert/ConcertApp/MultiUserClient/Source/MultiUserClient/Private/Widgets/ActiveSession/Replication/Client/Multi/Columns/SAssignPropertyComboBox.cpp
@@ -10,9 +10,10 @@
 #include "Replication/Editor/Model/PropertyUtils.h"
 #include "Widgets/ActiveSession/Replication/Misc/SNoClients.h"
 #include "Widgets/ActiveSession/Replication/Client/ClientUtils.h"
-#include "Widgets/ClientName/SHorizontalClientList.h"
-#include "Widgets/ClientName/SLocalClientName.h"
-#include "Widgets/ClientName/SRemoteClientName.h"
+#include "Widgets/Client/ClientInfoHelpers.h"
+#include "Widgets/Client/SHorizontalClientList.h"
+#include "Widgets/Client/SLocalClientName.h"
+#include "Widgets/Client/SRemoteClientName.h"
 
 #include "Algo/AnyOf.h"
 #include "Framework/MultiBox/MultiBoxBuilder.h"
@@ -58,10 +59,12 @@ namespace UE::MultiUserClient
 	{
 		using SWidgetType = ConcertClientSharedSlate::SHorizontalClientList;
 		const TArray<FGuid> Clients = AssignPropertyComboBox::GetDisplayedClients(ClientManager, DisplayedProperty, EditedObjects);
+		const ConcertSharedSlate::FIsLocalClient IsLocalClientDelegate = ConcertClientSharedSlate::MakeIsLocalClientGetter(LocalConcertClient);
 		return SWidgetType::GetDisplayString(
-			LocalConcertClient.Get(),
 			Clients,
-			SWidgetType::FSortPredicate::CreateStatic(&SWidgetType::SortLocalClientFirstThenAlphabetical, LocalConcertClient)
+			ConcertClientSharedSlate::MakeClientInfoGetter(LocalConcertClient),
+			SWidgetType::FSortPredicate::CreateStatic(&SWidgetType::SortLocalClientFirstThenAlphabetical, IsLocalClientDelegate),
+			IsLocalClientDelegate
 			);
 	}
 
@@ -88,7 +91,9 @@ namespace UE::MultiUserClient
 			.HasDownArrow(true)
 			.ButtonContent()
 			[
-				SAssignNew(ClientListWidget, ConcertClientSharedSlate::SHorizontalClientList, ConcertClient.ToSharedRef())
+				SAssignNew(ClientListWidget, ConcertClientSharedSlate::SHorizontalClientList)
+				.IsLocalClient(ConcertClientSharedSlate::MakeIsLocalClientGetter(ConcertClient.ToSharedRef()))
+				.GetClientInfo(ConcertClientSharedSlate::MakeClientInfoGetter(ConcertClient.ToSharedRef()))
 				.Font(FAppStyle::GetFontStyle("PropertyWindow.NormalFont"))
 				.HighlightText_Lambda([this](){ return HighlightText ? *HighlightText : FText::GetEmpty(); })
 				.EmptyListSlot() [ SNew(SNoClients) ]
@@ -117,13 +122,14 @@ namespace UE::MultiUserClient
 			const bool bIsLocalClient = EndpointId == ConcertClient->GetCurrentSession()->GetSessionClientEndpointId();
 			if (bIsLocalClient)
 			{
-				return SNew(SLocalClientName, ConcertClient.ToSharedRef())
+				return SNew(SLocalClientName)
+					.DisplayInfo(ConcertClientSharedSlate::MakeLocalClientInfoAttribute(ConcertClient.ToSharedRef()))
 					.Font(FAppStyle::GetFontStyle("PropertyWindow.NormalFont"))
 					.HighlightText_Lambda([this](){ return HighlightText ? *HighlightText : FText::GetEmpty(); });
 			}
-			return SNew(SRemoteClientName, ConcertClient.ToSharedRef())
+			return SNew(SRemoteClientName)
+				.DisplayInfo(ConcertClientSharedSlate::MakeClientInfoAttribute(ConcertClient.ToSharedRef(), EndpointId))
 				.Font(FAppStyle::GetFontStyle("PropertyWindow.NormalFont"))
-				.ClientEndpointId(EndpointId)
 				.HighlightText_Lambda([this](){ return HighlightText ? *HighlightText : FText::GetEmpty(); });
 		};
 		
diff --git a/Engine/Plugins/Developer/Concert/ConcertApp/MultiUserClient/Source/MultiUserClient/Private/Widgets/ActiveSession/Replication/Client/Multi/Columns/SReassignObjectComboBox.cpp b/Engine/Plugins/Developer/Concert/ConcertApp/MultiUserClient/Source/MultiUserClient/Private/Widgets/ActiveSession/Replication/Client/Multi/Columns/SReassignObjectComboBox.cpp
--- a/Engine/Plugins/Developer/Concert/ConcertApp/MultiUserClient/Source/MultiUserClient/Private/Widgets/ActiveSession/Replication/Client/Multi/Columns/SReassignObjectComboBox.cpp
+++ b/Engine/Plugins/Developer/Concert/ConcertApp/MultiUserClient/Source/MultiUserClient/Private/Widgets/ActiveSession/Replication/Client/Multi/Columns/SReassignObjectComboBox.cpp
@@ -5,25 +5,25 @@
 #include "MultiUserReplicationStyle.h"
 #include "Misc/EBreakBehavior.h"
 #include "Replication/Client/ReplicationClient.h"
+#include "Replication/Editor/Model/Object/IObjectHierarchyModel.h"
 #include "Replication/Submission/MultiEdit/ReassignObjectPropertiesLogic.h"
 #include "Widgets/ActiveSession/Replication/Client/ClientUtils.h"
 #include "Widgets/ActiveSession/Replication/Misc/SNoClients.h"
-#include "Widgets/ClientName/SHorizontalClientList.h"
+#include "Widgets/Client/ClientInfoHelpers.h"
+#include "Widgets/Client/SHorizontalClientList.h"
 
 #include "Framework/MultiBox/MultiBoxBuilder.h"
 #include "Framework/Commands/UIAction.h"
 #include "Framework/Commands/UICommandInfo.h"
-#include "Replication/Editor/Model/Object/IObjectHierarchyModel.h"
 #include "Textures/SlateIcon.h"
 #include "Widgets/Images/SThrobber.h"
 #include "Widgets/Input/SComboBox.h"
-#include "Widgets/Layout/SScaleBox.h"
 
 #define LOCTEXT_NAMESPACE "SReassignObjectComboBox"
 
 namespace UE::MultiUserClient
 {
-	namespace ReassignObjctComboBox
+	namespace ReassignObjectComboBox
 	{
 		TArray<FGuid> GetDisplayedClients(
 			const FReassignObjectPropertiesLogic& ReassignmentLogic,
@@ -50,7 +50,7 @@ namespace UE::MultiUserClient
 		TArray<FString>& InOutSearchTerms
 		)
 	{
-		for (const FGuid& ClientId : ReassignObjctComboBox::GetDisplayedClients(ReassignmentLogic, ManagedObject))
+		for (const FGuid& ClientId : ReassignObjectComboBox::GetDisplayedClients(ReassignmentLogic, ManagedObject))
 		{
 			InOutSearchTerms.Add(ClientUtils::GetClientDisplayName(Session, ClientId));
 		}
@@ -63,11 +63,13 @@ namespace UE::MultiUserClient
 		)
 	{
 		using SWidgetType = ConcertClientSharedSlate::SHorizontalClientList;
-		const TArray<FGuid> Clients = ReassignObjctComboBox::GetDisplayedClients(ReassignmentLogic, ManagedObject);
+		const TArray<FGuid> Clients = ReassignObjectComboBox::GetDisplayedClients(ReassignmentLogic, ManagedObject);
+		const ConcertSharedSlate::FIsLocalClient IsLocalClientDelegate = ConcertClientSharedSlate::MakeIsLocalClientGetter(LocalConcertClient);
 		return SWidgetType::GetDisplayString(
-			LocalConcertClient.Get(),
 			Clients,
-			SWidgetType::FSortPredicate::CreateStatic(&SWidgetType::SortLocalClientFirstThenAlphabetical, LocalConcertClient)
+			ConcertClientSharedSlate::MakeClientInfoGetter(LocalConcertClient),
+			SWidgetType::FSortPredicate::CreateStatic(&SWidgetType::SortLocalClientFirstThenAlphabetical, IsLocalClientDelegate),
+			IsLocalClientDelegate
 			);
 	}
 
@@ -107,7 +109,9 @@ namespace UE::MultiUserClient
 				.HasDownArrow(true)
 				.ButtonContent()
 				[
-					SAssignNew(ComboClientList, ConcertClientSharedSlate::SHorizontalClientList, InConcertClient)
+					SAssignNew(ComboClientList, ConcertClientSharedSlate::SHorizontalClientList)
+					.IsLocalClient(ConcertClientSharedSlate::MakeIsLocalClientGetter(ConcertClient.ToSharedRef()))
+					.GetClientInfo(ConcertClientSharedSlate::MakeClientInfoGetter(ConcertClient.ToSharedRef()))
 					.Font(FAppStyle::GetFontStyle("PropertyWindow.NormalFont"))
 					.HighlightText_Lambda([this](){ return HighlightText ? *HighlightText : FText::GetEmpty(); })
 					.EmptyListSlot() [ SNew(SNoClients) ]
@@ -128,7 +132,7 @@ namespace UE::MultiUserClient
 	void SReassignObjectComboBox::UpdateComboButtonContent() const
 	{
 		ComboClientList->RefreshList(
-			ReassignObjctComboBox::GetDisplayedClients(*ReassignmentLogic, ManagedObject)
+			ReassignObjectComboBox::GetDisplayedClients(*ReassignmentLogic, ManagedObject)
 			);
 	}
 
diff --git a/Engine/Plugins/Developer/Concert/ConcertApp/MultiUserClient/Source/MultiUserClient/Private/Widgets/ActiveSession/Replication/Joined/SSelectClientViewComboButton.cpp b/Engine/Plugins/Developer/Concert/ConcertApp/MultiUserClient/Source/MultiUserClient/Private/Widgets/ActiveSession/Replication/Joined/SSelectClientViewComboButton.cpp
deleted file mode 100644
diff --git a/Engine/Plugins/Developer/Concert/ConcertApp/MultiUserClient/Source/MultiUserClient/Private/Widgets/ActiveSession/Replication/Joined/SSelectClientViewComboButton.h b/Engine/Plugins/Developer/Concert/ConcertApp/MultiUserClient/Source/MultiUserClient/Private/Widgets/ActiveSession/Replication/Joined/SSelectClientViewComboButton.h
deleted file mode 100644
diff --git a/Engine/Plugins/Developer/Concert/ConcertApp/MultiUserClient/Source/MultiUserClient/Private/Widgets/Disconnected/SConcertClientSessionBrowser.cpp b/Engine/Plugins/Developer/Concert/ConcertApp/MultiUserClient/Source/MultiUserClient/Private/Widgets/Disconnected/SConcertClientSessionBrowser.cpp
--- a/Engine/Plugins/Developer/Concert/ConcertApp/MultiUserClient/Source/MultiUserClient/Private/Widgets/Disconnected/SConcertClientSessionBrowser.cpp
+++ b/Engine/Plugins/Developer/Concert/ConcertApp/MultiUserClient/Source/MultiUserClient/Private/Widgets/Disconnected/SConcertClientSessionBrowser.cpp
@@ -16,6 +16,9 @@
 #include "Session/Browser/ConcertBrowserUtils.h"
 #include "Session/Browser/Items/ConcertSessionTreeItem.h"
 #include "Session/Browser/SConcertSessionBrowser.h"
+#include "Widgets/Disconnected/ConcertClientSessionBrowserController.h"
+#include "Widgets/Client/SClientName.h"
+#include "Widgets/Client/ClientInfoHelpers.h"
 
 #include "Algo/Transform.h"
 #include "Algo/TiedTupleOutput.h"
@@ -25,7 +28,6 @@
 #include "Misc/MessageDialog.h"
 #include "Styling/AppStyle.h"
 #include "Textures/SlateIcon.h"
-#include "Widgets/Disconnected/ConcertClientSessionBrowserController.h"
 #include "Widgets/Images/SImage.h"
 #include "Widgets/Input/SButton.h"
 #include "Widgets/Layout/SBorder.h"
@@ -110,7 +112,7 @@ void SConcertClientSessionBrowser::Construct(const FArguments& InArgs, IConcertC
 
 	ChildSlot
 	[
-		MakeBrowserContent(InSearchText)
+		MakeBrowserContent(InSearchText, InConcertClient.ToSharedRef())
 	];
 
 	// Create a timer to periodically poll the server for sessions and session clients at a lower frequency than the normal tick.
@@ -120,7 +122,7 @@ void SConcertClientSessionBrowser::Construct(const FArguments& InArgs, IConcertC
 	ScheduleDiscoveryTaskIfPossible(bForceTaskStart);
 }
 
-TSharedRef<SWidget> SConcertClientSessionBrowser::MakeBrowserContent(TSharedPtr<FText> InSearchText)
+TSharedRef<SWidget> SConcertClientSessionBrowser::MakeBrowserContent(TSharedPtr<FText> InSearchText, const TSharedRef<IConcertClient>& Client)
 {
 	return SNew(SBox)
 		.HAlign(HAlign_Fill)
@@ -139,7 +141,7 @@ TSharedRef<SWidget> SConcertClientSessionBrowser::MakeBrowserContent(TSharedPtr<
 					.ExtendSessionContextMenu(this, &SConcertClientSessionBrowser::ExtendSessionContextMenu)
 					.RightOfControlButtons()
 					[
-						MakeUserAndSettings()
+						MakeUserAndSettings(InSearchText, Client)
 					]
 					.OnSessionClicked(this, &SConcertClientSessionBrowser::OnSessionSelectionChanged)
 					.OnLiveSessionDoubleClicked(this, &SConcertClientSessionBrowser::OnSessionDoubleClicked)
@@ -286,27 +288,17 @@ void SConcertClientSessionBrowser::ExtendSessionContextMenu(const TSharedPtr<FCo
 	
 }
 
-TSharedRef<SWidget> SConcertClientSessionBrowser::MakeUserAndSettings()
+TSharedRef<SWidget> SConcertClientSessionBrowser::MakeUserAndSettings(TSharedPtr<FText> InSearchText, const TSharedRef<IConcertClient>& Client)
 {
 	return SNew(SHorizontalBox)
-		// The user "Avatar color" displayed as a small square colored by the user avatar color.
 		+SHorizontalBox::Slot()
 		.VAlign(VAlign_Center)
 		.AutoWidth()
 		[
-			SNew(SImage)
-			.ColorAndOpacity_Lambda([this]() { return Controller->GetConcertClient()->GetClientInfo().AvatarColor; })
-			.Image(FAppStyle::GetBrush("Icons.FilledCircle"))
-		]
-				
-		// The user "Display Name".
-		+SHorizontalBox::Slot()
-		.VAlign(VAlign_Center)
-		.AutoWidth()
-		.Padding(3, 0, 2, 0)
-		[
-			SNew(STextBlock)
-				.Text_Lambda([this]() { return FText::FromString(Controller->GetConcertClient()->GetClientInfo().DisplayName);} )
+			SNew(UE::ConcertClientSharedSlate::SClientName)
+			.ClientInfo(UE::ConcertClientSharedSlate::MakeLocalClientInfoAttribute(Client))
+			.DisplayAsLocalClient(false) // It is obvious that it's the local user because we have not joined any session yet so do not display (You) in the back.
+			.HighlightText_Lambda([InSearchText](){ return *InSearchText.Get(); })
 		]
 
 		// The "Settings" icons.
diff --git a/Engine/Plugins/Developer/Concert/ConcertApp/MultiUserClient/Source/MultiUserClient/Private/Widgets/Disconnected/SConcertClientSessionBrowser.h b/Engine/Plugins/Developer/Concert/ConcertApp/MultiUserClient/Source/MultiUserClient/Private/Widgets/Disconnected/SConcertClientSessionBrowser.h
--- a/Engine/Plugins/Developer/Concert/ConcertApp/MultiUserClient/Source/MultiUserClient/Private/Widgets/Disconnected/SConcertClientSessionBrowser.h
+++ b/Engine/Plugins/Developer/Concert/ConcertApp/MultiUserClient/Source/MultiUserClient/Private/Widgets/Disconnected/SConcertClientSessionBrowser.h
@@ -71,10 +71,10 @@ private:
 	TSharedPtr<SWidget> NoClientPanel; // Displays the 'no client connected' message in Clients expendable area.
 	
 	// Layout the 'session|details' split view.
-	TSharedRef<SWidget> MakeBrowserContent(TSharedPtr<FText> InSearchText);
+	TSharedRef<SWidget> MakeBrowserContent(TSharedPtr<FText> InSearchText, const TSharedRef<IConcertClient>& Client);
 	void ExtendControlButtons(FExtender& Extender);
 	void ExtendSessionContextMenu(const TSharedPtr<FConcertSessionTreeItem>& Item, FExtender& Extender);
-	TSharedRef<SWidget> MakeUserAndSettings();
+	TSharedRef<SWidget> MakeUserAndSettings(TSharedPtr<FText> InSearchText, const TSharedRef<IConcertClient>& Client);
 	TSharedRef<SWidget> MakeOverlayedTableView(const TSharedRef<SWidget>& TableView);
 	
 	// Layouts the session detail panel.
diff --git a/Engine/Plugins/Developer/Concert/ConcertUI/ConcertClientSharedSlate/Source/ConcertClientSharedSlate/Private/Widgets/ClientName/SClientName.cpp b/Engine/Plugins/Developer/Concert/ConcertUI/ConcertClientSharedSlate/Source/ConcertClientSharedSlate/Private/Widgets/ClientName/SClientName.cpp
deleted file mode 100644
diff --git a/Engine/Plugins/Developer/Concert/ConcertUI/ConcertClientSharedSlate/Source/ConcertClientSharedSlate/Private/Widgets/ClientName/SHorizontalClientList.cpp b/Engine/Plugins/Developer/Concert/ConcertUI/ConcertClientSharedSlate/Source/ConcertClientSharedSlate/Private/Widgets/ClientName/SHorizontalClientList.cpp
deleted file mode 100644
diff --git a/Engine/Plugins/Developer/Concert/ConcertUI/ConcertClientSharedSlate/Source/ConcertClientSharedSlate/Private/Widgets/ClientName/SLocalClientName.cpp b/Engine/Plugins/Developer/Concert/ConcertUI/ConcertClientSharedSlate/Source/ConcertClientSharedSlate/Private/Widgets/ClientName/SLocalClientName.cpp
deleted file mode 100644
diff --git a/Engine/Plugins/Developer/Concert/ConcertUI/ConcertClientSharedSlate/Source/ConcertClientSharedSlate/Private/Widgets/ClientName/SRemoteClientName.cpp b/Engine/Plugins/Developer/Concert/ConcertUI/ConcertClientSharedSlate/Source/ConcertClientSharedSlate/Private/Widgets/ClientName/SRemoteClientName.cpp
deleted file mode 100644
diff --git a/Engine/Plugins/Developer/Concert/ConcertUI/ConcertClientSharedSlate/Source/ConcertClientSharedSlate/Public/Widgets/Client/ClientInfoHelpers.h b/Engine/Plugins/Developer/Concert/ConcertUI/ConcertClientSharedSlate/Source/ConcertClientSharedSlate/Public/Widgets/Client/ClientInfoHelpers.h
new file mode 100644
--- /dev/null
+++ b/Engine/Plugins/Developer/Concert/ConcertUI/ConcertClientSharedSlate/Source/ConcertClientSharedSlate/Public/Widgets/Client/ClientInfoHelpers.h
@@ -0,0 +1,88 @@
+﻿// Copyright Epic Games, Inc. All Rights Reserved.
+
+#pragma once
+
+#include "ConcertMessageData.h"
+#include "IConcertClient.h"
+#include "IConcertSession.h"
+#include "Widgets/Client/ClientInfoDelegate.h"
+
+#include "Misc/Optional.h"
+
+/** Util types that can be used with SClientName, etc. */
+namespace UE::ConcertClientSharedSlate
+{
+	inline ConcertSharedSlate::FGetOptionalClientInfo MakeClientInfoGetter(const TSharedRef<IConcertClient>& Client)
+	{
+		return ConcertSharedSlate::FGetOptionalClientInfo::CreateLambda([WeakClient = Client.ToWeakPtr()](const FGuid& ClientEndpointId) -> TOptional<FConcertClientInfo>
+		{
+			TSharedPtr<IConcertClient> ClientPin = WeakClient.Pin();
+			if (!ClientPin)
+			{
+				return {};
+			}
+
+			const TSharedPtr<IConcertClientSession> Session = ClientPin->GetCurrentSession();
+			if (!Session)
+			{
+				return {};
+			}
+			
+			// FindSessionClient does not work for the local client
+			if (ClientEndpointId == Session->GetSessionClientEndpointId())
+			{
+				return ClientPin->GetClientInfo();
+			}
+			
+			FConcertSessionClientInfo ClientInfo;
+			Session->FindSessionClient(ClientEndpointId, ClientInfo);
+			return ClientInfo.ClientInfo;
+		});
+	}
+	
+	inline ConcertSharedSlate::FIsLocalClient MakeIsLocalClientGetter(const TSharedRef<IConcertClient>& Client)
+	{
+		return ConcertSharedSlate::FIsLocalClient::CreateLambda([WeakClient = Client.ToWeakPtr()](const FGuid& ClientEndpointId)
+		{
+			const TSharedPtr<IConcertClient> ClientPin = WeakClient.Pin();
+			return ClientPin && ClientPin->GetCurrentSession() && ClientPin->GetCurrentSession()->GetSessionClientEndpointId() == ClientEndpointId;
+		});
+	}
+
+	inline TAttribute<TOptional<FConcertClientInfo>> MakeLocalClientInfoAttribute(const TSharedRef<IConcertClient>& Client)
+	{
+		return TAttribute<TOptional<FConcertClientInfo>>::CreateLambda([WeakClient = Client.ToWeakPtr()]()
+		{
+			const TSharedPtr<IConcertClient> ClientPin = WeakClient.Pin();
+			return ClientPin ? ClientPin->GetClientInfo() : TOptional<FConcertClientInfo>{};
+		});
+	}
+
+	inline TAttribute<TOptional<FConcertClientInfo>> MakeClientInfoAttribute(const TSharedRef<IConcertClient>& Client, const FGuid& ClientId)
+	{
+		return TAttribute<TOptional<FConcertClientInfo>>::CreateLambda([WeakClient = Client.ToWeakPtr(), ClientId]() -> TOptional<FConcertClientInfo>
+		{
+			const TSharedPtr<IConcertClient> ClientPin = WeakClient.Pin();
+			if (!ClientPin)
+			{
+				return {};
+			}
+
+			const TSharedPtr<IConcertClientSession> Session = ClientPin->GetCurrentSession();
+			if (!Session)
+			{
+				return {};
+			}
+
+			// FindSessionClient does not work for the local client
+			if (ClientId == Session->GetSessionClientEndpointId())
+			{
+				return ClientPin->GetClientInfo();
+			}
+			
+			FConcertSessionClientInfo Info;
+			Session->FindSessionClient(ClientId, Info);
+			return Info.ClientInfo;
+		});
+	}
+}
\ No newline at end of file
diff --git a/Engine/Plugins/Developer/Concert/ConcertUI/ConcertClientSharedSlate/Source/ConcertClientSharedSlate/Public/Widgets/ClientName/SClientName.h b/Engine/Plugins/Developer/Concert/ConcertUI/ConcertClientSharedSlate/Source/ConcertClientSharedSlate/Public/Widgets/ClientName/SClientName.h
deleted file mode 100644
diff --git a/Engine/Plugins/Developer/Concert/ConcertUI/ConcertClientSharedSlate/Source/ConcertClientSharedSlate/Public/Widgets/ClientName/SHorizontalClientList.h b/Engine/Plugins/Developer/Concert/ConcertUI/ConcertClientSharedSlate/Source/ConcertClientSharedSlate/Public/Widgets/ClientName/SHorizontalClientList.h
deleted file mode 100644
diff --git a/Engine/Plugins/Developer/Concert/ConcertUI/ConcertClientSharedSlate/Source/ConcertClientSharedSlate/Public/Widgets/ClientName/SLocalClientName.h b/Engine/Plugins/Developer/Concert/ConcertUI/ConcertClientSharedSlate/Source/ConcertClientSharedSlate/Public/Widgets/ClientName/SLocalClientName.h
deleted file mode 100644
diff --git a/Engine/Plugins/Developer/Concert/ConcertUI/ConcertClientSharedSlate/Source/ConcertClientSharedSlate/Public/Widgets/ClientName/SRemoteClientName.h b/Engine/Plugins/Developer/Concert/ConcertUI/ConcertClientSharedSlate/Source/ConcertClientSharedSlate/Public/Widgets/ClientName/SRemoteClientName.h
deleted file mode 100644
diff --git a/Engine/Plugins/Developer/Concert/ConcertUI/ConcertSharedSlate/Source/ConcertSharedSlate/Private/Widgets/Client/SClientName.cpp b/Engine/Plugins/Developer/Concert/ConcertUI/ConcertSharedSlate/Source/ConcertSharedSlate/Private/Widgets/Client/SClientName.cpp
new file mode 100644
--- /dev/null
+++ b/Engine/Plugins/Developer/Concert/ConcertUI/ConcertSharedSlate/Source/ConcertSharedSlate/Private/Widgets/Client/SClientName.cpp
@@ -0,0 +1,83 @@
+// Copyright Epic Games, Inc. All Rights Reserved.
+
+#include "Widgets/Client/SClientName.h"
+
+#include "Styling/AppStyle.h"
+#include "Widgets/Images/SImage.h"
+#include "Widgets/Layout/SBorder.h"
+#include "Widgets/SBoxPanel.h"
+#include "Widgets/Text/STextBlock.h"
+
+#define LOCTEXT_NAMESPACE "SClientName"
+
+namespace UE::ConcertClientSharedSlate
+{
+	void SClientName::Construct(const FArguments& InArgs)
+	{
+		ClientInfoAttribute = InArgs._ClientInfo;
+		DisplayAsLocalClientAttribute = InArgs._DisplayAsLocalClient;
+		check(ClientInfoAttribute.IsSet() || ClientInfoAttribute.IsBound());
+		check(DisplayAsLocalClientAttribute.IsSet() || DisplayAsLocalClientAttribute.IsBound());
+		
+		ChildSlot
+		[
+			SNew(SHorizontalBox)
+			
+			// The user "Avatar color" displayed as a small square colored by the user avatar color.
+			+SHorizontalBox::Slot()
+			.VAlign(VAlign_Center)
+			.AutoWidth()
+			[
+				SNew(SImage)
+				.ColorAndOpacity(this, &SClientName::GetAvatarColor)
+				.Image(FAppStyle::GetBrush("Icons.FilledCircle"))
+			]
+					
+			// The user "Display Name".
+			+SHorizontalBox::Slot()
+			.VAlign(VAlign_Center)
+			.AutoWidth()
+			.Padding(1, 0, 0, 0)
+			[
+				SNew(SBorder)
+				.BorderImage(FAppStyle::Get().GetBrush("NoBorder"))
+				.ColorAndOpacity(FLinearColor(0.75f, 0.75f, 0.75f))
+				[
+					SNew(STextBlock)
+					.Font(InArgs._Font)
+					.Text(TAttribute<FText>::CreateSP(this, &SClientName::GetClientDisplayName))
+					.HighlightText(InArgs._HighlightText)
+				]
+			]
+		];
+	}
+
+	FText SClientName::GetDisplayText(const FConcertClientInfo& Info, bool bDisplayAsLocalClient)
+	{
+		if (bDisplayAsLocalClient)
+		{
+			return FText::Format(
+				LOCTEXT("ClientDisplayNameFmt", "{0} (You)"),
+				FText::FromString(Info.DisplayName)
+				);
+		}
+		
+		return FText::FromString(Info.DisplayName);
+	}
+
+	FText SClientName::GetClientDisplayName() const
+	{
+		const TOptional<FConcertClientInfo> ClientInfo = ClientInfoAttribute.Get();
+		return ClientInfo
+			? GetDisplayText(*ClientInfo, DisplayAsLocalClientAttribute.Get())
+			: LOCTEXT("Unavailable", "Unavailable");
+	}
+
+	FSlateColor SClientName::GetAvatarColor() const
+	{
+		const TOptional<FConcertClientInfo> ClientInfo = ClientInfoAttribute.Get();
+		return ClientInfo ? ClientInfo->AvatarColor : FSlateColor(FLinearColor::Gray);
+	}
+}
+
+#undef LOCTEXT_NAMESPACE
\ No newline at end of file
diff --git a/Engine/Plugins/Developer/Concert/ConcertUI/ConcertSharedSlate/Source/ConcertSharedSlate/Private/Widgets/Client/SHorizontalClientList.cpp b/Engine/Plugins/Developer/Concert/ConcertUI/ConcertSharedSlate/Source/ConcertSharedSlate/Private/Widgets/Client/SHorizontalClientList.cpp
new file mode 100644
--- /dev/null
+++ b/Engine/Plugins/Developer/Concert/ConcertUI/ConcertSharedSlate/Source/ConcertSharedSlate/Private/Widgets/Client/SHorizontalClientList.cpp
@@ -0,0 +1,164 @@
+// Copyright Epic Games, Inc. All Rights Reserved.
+
+#include "Widgets/Client/SHorizontalClientList.h"
+
+#include "Widgets/Client/SClientName.h"
+
+#include "Widgets/Client/SLocalClientName.h"
+#include "Widgets/Client/SRemoteClientName.h"
+#include "Widgets/Layout/SScrollBox.h"
+#include "Widgets/Layout/SWidgetSwitcher.h"
+#include "Widgets/Text/STextBlock.h"
+
+#define LOCTEXT_NAMESPACE "SHorizontalClientList"
+
+namespace UE::ConcertClientSharedSlate
+{
+	namespace HorizontalClientList
+	{
+		TArray<FConcertSessionClientInfo> GetSortedClients(
+			const TConstArrayView<FGuid>& Clients,
+			const ConcertSharedSlate::FGetOptionalClientInfo& GetClientInfoDelegate,
+			const SHorizontalClientList::FSortPredicate& SortPredicate
+			)
+		{
+			// Prefetch the client info to avoid many GetClientInfoDelegate calls during Sort()
+			TArray<FConcertSessionClientInfo> ClientsToDisplay;
+			for (const FGuid& Client : Clients)
+			{
+				const TOptional<FConcertClientInfo> ClientInfo = GetClientInfoDelegate.Execute(Client);
+				FConcertSessionClientInfo Info;
+				Info.ClientEndpointId = Client;
+				Info.ClientInfo = ClientInfo ? *ClientInfo : FConcertClientInfo{ .DisplayName = TEXT("Unavailable") };
+				ClientsToDisplay.Add(Info);
+			}
+
+			ClientsToDisplay.Sort([&SortPredicate](const FConcertSessionClientInfo& Left, const FConcertSessionClientInfo& Right)
+			{
+				return SortPredicate.Execute(Left, Right);
+			});
+			return ClientsToDisplay;
+		}
+	}
+	
+	bool SHorizontalClientList::SortLocalClientFirstThenAlphabetical(const FConcertSessionClientInfo& Left, const FConcertSessionClientInfo& Right, ConcertSharedSlate::FIsLocalClient IsLocalClientDelegate)
+	{
+		// If one of the compare clients is local, always return that the local client is smaller.
+		const bool bLeftIsLocalClient = IsLocalClientDelegate.IsBound() && IsLocalClientDelegate.Execute(Left.ClientEndpointId);
+		const bool bRightIsLocalClient = IsLocalClientDelegate.IsBound() && IsLocalClientDelegate.Execute(Right.ClientEndpointId);
+		return bLeftIsLocalClient
+			|| (!bRightIsLocalClient && Left.ClientInfo.DisplayName < Right.ClientInfo.DisplayName);
+	}
+
+	TOptional<FString> SHorizontalClientList::GetDisplayString(
+		const TConstArrayView<FGuid>& Clients,
+		const ConcertSharedSlate::FGetOptionalClientInfo& GetClientInfoDelegate,
+		const FSortPredicate& SortPredicate,
+		const ConcertSharedSlate::FIsLocalClient& IsLocalClientDelegate
+		)
+	{
+		const TArray<FConcertSessionClientInfo> ClientsToDisplay = HorizontalClientList::GetSortedClients(Clients, GetClientInfoDelegate, SortPredicate);
+		if (!ClientsToDisplay.IsEmpty())
+		{
+			return {};
+		}
+		
+		return FString::JoinBy(ClientsToDisplay, TEXT(", "), [&IsLocalClientDelegate](const FConcertSessionClientInfo& ClientInfo)
+			{
+				// GetSortedClients should return empty if GetSortedClients is invalid
+				const bool bIsLocalClient = IsLocalClientDelegate.IsBound() && IsLocalClientDelegate.Execute(ClientInfo.ClientEndpointId);
+				return SClientName::GetDisplayText(ClientInfo.ClientInfo, bIsLocalClient).ToString();
+			});
+	}
+
+	void SHorizontalClientList::Construct(const FArguments& InArgs)
+	{
+		IsLocalClientDelegate = InArgs._IsLocalClient;
+		GetClientInfoDelegate = InArgs._GetClientInfo;
+		SortPredicateDelegate = InArgs._SortPredicate.IsBound()
+			? InArgs._SortPredicate
+			: FSortPredicate::CreateStatic(&SHorizontalClientList::SortLocalClientFirstThenAlphabetical, IsLocalClientDelegate);
+		
+		DisplayAvatarColorAttribute = InArgs._DisplayAvatarColor;
+		HighlightTextAttribute = InArgs._HighlightText;
+		
+		NameFont = InArgs._Font;
+
+		check(GetClientInfoDelegate.IsBound());
+		
+		ChildSlot
+		[
+			SAssignNew(WidgetSwitcher, SWidgetSwitcher)
+			.WidgetIndex(0)
+			+SWidgetSwitcher::Slot()
+			[
+				InArgs._EmptyListSlot.Widget
+			]
+			+SWidgetSwitcher::Slot()
+			[
+				SAssignNew(ScrollBox, SScrollBox)
+				.Orientation(Orient_Horizontal)
+			]
+		];
+	}
+
+	void SHorizontalClientList::RefreshList(const TConstArrayView<FGuid>& Clients)
+	{
+		ScrollBox->ClearChildren();
+
+		if (Clients.IsEmpty())
+		{
+			WidgetSwitcher->SetActiveWidgetIndex(0);
+			return;
+		}
+		WidgetSwitcher->SetActiveWidgetIndex(1);
+
+		const TArray<FConcertSessionClientInfo> ClientsToDisplay = HorizontalClientList::GetSortedClients(Clients, GetClientInfoDelegate, SortPredicateDelegate);
+		
+		bool bIsFirst = true;
+		for (const FConcertSessionClientInfo& Info : ClientsToDisplay)
+		{
+			if (!bIsFirst)
+			{
+				ScrollBox->AddSlot()
+				.VAlign(VAlign_Center)
+				.HAlign(HAlign_Left)
+				.Padding(-1, 1, 0, 0)
+				[
+					SNew(STextBlock)
+					.Text(LOCTEXT("Comma", ", "))
+					.Font(NameFont)
+				];
+			}
+
+			const FGuid& EndpointId = Info.ClientEndpointId;
+			const bool bIsLocalClient = IsLocalClientDelegate.IsBound() && IsLocalClientDelegate.Execute(Info.ClientEndpointId);
+			if (bIsLocalClient)
+			{
+				ScrollBox->AddSlot()
+				[
+					SNew(SLocalClientName)
+					.DisplayInfo_Lambda([this, EndpointId](){ return GetClientInfoDelegate.Execute(EndpointId); })
+					.DisplayAvatarColor(DisplayAvatarColorAttribute)
+					.HighlightText(HighlightTextAttribute)
+					.Font(NameFont)
+				];
+			}
+			else
+			{
+				ScrollBox->AddSlot()
+				[
+					SNew(SRemoteClientName)
+					.DisplayInfo_Lambda([this, EndpointId](){ return GetClientInfoDelegate.Execute(EndpointId); })
+					.DisplayAvatarColor(DisplayAvatarColorAttribute)
+					.HighlightText(HighlightTextAttribute)
+					.Font(NameFont)
+				];
+			}
+
+			bIsFirst = false;
+		}
+	}
+}
+
+#undef LOCTEXT_NAMESPACE
\ No newline at end of file
diff --git a/Engine/Plugins/Developer/Concert/ConcertUI/ConcertSharedSlate/Source/ConcertSharedSlate/Private/Widgets/Client/SLocalClientName.cpp b/Engine/Plugins/Developer/Concert/ConcertUI/ConcertSharedSlate/Source/ConcertSharedSlate/Private/Widgets/Client/SLocalClientName.cpp
new file mode 100644
--- /dev/null
+++ b/Engine/Plugins/Developer/Concert/ConcertUI/ConcertSharedSlate/Source/ConcertSharedSlate/Private/Widgets/Client/SLocalClientName.cpp
@@ -0,0 +1,26 @@
+// Copyright Epic Games, Inc. All Rights Reserved.
+
+#include "Widgets/Client/SLocalClientName.h"
+
+#include "Widgets/Client/SClientName.h"
+#include "Widgets/Text/STextBlock.h"
+
+#define LOCTEXT_NAMESPACE "SLocalClientName"
+
+namespace UE::ConcertClientSharedSlate
+{
+	void SLocalClientName::Construct(const FArguments& InArgs)
+	{
+		ChildSlot
+		[
+			SNew(SClientName)
+			.ClientInfo(InArgs._DisplayInfo)
+			.DisplayAsLocalClient(true)
+			.DisplayAvatarColor(InArgs._DisplayAvatarColor)
+			.HighlightText(InArgs._HighlightText)
+			.Font(InArgs._Font)
+		];
+	}
+}
+
+#undef LOCTEXT_NAMESPACE
\ No newline at end of file
diff --git a/Engine/Plugins/Developer/Concert/ConcertUI/ConcertSharedSlate/Source/ConcertSharedSlate/Private/Widgets/Client/SRemoteClientName.cpp b/Engine/Plugins/Developer/Concert/ConcertUI/ConcertSharedSlate/Source/ConcertSharedSlate/Private/Widgets/Client/SRemoteClientName.cpp
new file mode 100644
--- /dev/null
+++ b/Engine/Plugins/Developer/Concert/ConcertUI/ConcertSharedSlate/Source/ConcertSharedSlate/Private/Widgets/Client/SRemoteClientName.cpp
@@ -0,0 +1,41 @@
+// Copyright Epic Games, Inc. All Rights Reserved.
+
+#include "Widgets/Client/SRemoteClientName.h"
+
+#include "Widgets/Client/SClientName.h"
+#include "Widgets/Text/STextBlock.h"
+
+#define LOCTEXT_NAMESPACE "SLocalClientName"
+
+namespace UE::ConcertClientSharedSlate
+{
+	void SRemoteClientName::Construct(const FArguments& InArgs)
+	{
+		ClientDisplayInfo = InArgs._DisplayInfo;
+		
+		ChildSlot
+		[
+			SNew(SClientName)
+			.ClientInfo(this, &SRemoteClientName::GetClientInfo)
+			.DisplayAsLocalClient(false)
+			.DisplayAvatarColor(InArgs._DisplayAvatarColor)
+			.HighlightText(InArgs._HighlightText)
+			.Font(InArgs._Font)
+		];
+	}
+
+	TOptional<FConcertClientInfo> SRemoteClientName::GetClientInfo() const
+	{
+		const TOptional<FConcertClientInfo> NewInfo = ClientDisplayInfo.IsBound()
+			? ClientDisplayInfo.Get()
+			: TOptional<FConcertClientInfo>{};
+					
+		if (NewInfo)
+		{
+			LastKnownClientInfo = NewInfo;
+		}
+		return LastKnownClientInfo;
+	}
+}
+
+#undef LOCTEXT_NAMESPACE
\ No newline at end of file
diff --git a/Engine/Plugins/Developer/Concert/ConcertUI/ConcertSharedSlate/Source/ConcertSharedSlate/Public/ConcertFrontendUtils.h b/Engine/Plugins/Developer/Concert/ConcertUI/ConcertSharedSlate/Source/ConcertSharedSlate/Public/ConcertFrontendUtils.h
--- a/Engine/Plugins/Developer/Concert/ConcertUI/ConcertSharedSlate/Source/ConcertSharedSlate/Public/ConcertFrontendUtils.h
+++ b/Engine/Plugins/Developer/Concert/ConcertUI/ConcertSharedSlate/Source/ConcertSharedSlate/Public/ConcertFrontendUtils.h
@@ -24,19 +24,6 @@ enum class ETimeFormat
 
 namespace ConcertFrontendUtils
 {
-	inline TSharedRef<SWidget> CreateDisplayName(const TAttribute<FText>& InDisplayName)
-	{
-		return SNew(SBorder)
-			.BorderImage(FAppStyle::Get().GetBrush("NoBorder"))
-			.ColorAndOpacity(FLinearColor(0.75f, 0.75f, 0.75f))
-			.Padding(FMargin(6.0f, 4.0f))
-			[
-				SNew(STextBlock)
-				.Font(FAppStyle::Get().GetFontStyle("BoldFont"))
-				.Text(InDisplayName)
-			];
-	}
-	
 	template <typename ItemType, typename PredFactoryType>
 	inline void SyncArraysByPredicate(TArray<TSharedPtr<ItemType>>& InOutArray, TArray<TSharedPtr<ItemType>>&& InNewArray, const PredFactoryType& InPredFactory)
 	{
diff --git a/Engine/Plugins/Developer/Concert/ConcertUI/ConcertSharedSlate/Source/ConcertSharedSlate/Public/Widgets/Client/ClientInfoDelegate.h b/Engine/Plugins/Developer/Concert/ConcertUI/ConcertSharedSlate/Source/ConcertSharedSlate/Public/Widgets/Client/ClientInfoDelegate.h
new file mode 100644
--- /dev/null
+++ b/Engine/Plugins/Developer/Concert/ConcertUI/ConcertSharedSlate/Source/ConcertSharedSlate/Public/Widgets/Client/ClientInfoDelegate.h
@@ -0,0 +1,15 @@
+﻿// Copyright Epic Games, Inc. All Rights Reserved.
+
+#pragma once
+
+#include "ConcertMessageData.h"
+#include "Delegates/Delegate.h"
+
+struct FGuid;
+template<typename OptionalType> struct TOptional;
+
+namespace UE::ConcertSharedSlate
+{
+	DECLARE_DELEGATE_RetVal_OneParam(bool, FIsLocalClient, const FGuid& ClientId);
+	DECLARE_DELEGATE_RetVal_OneParam(TOptional<FConcertClientInfo>, FGetOptionalClientInfo, const FGuid& ClientEndpointId);
+}
\ No newline at end of file
diff --git a/Engine/Plugins/Developer/Concert/ConcertUI/ConcertSharedSlate/Source/ConcertSharedSlate/Public/Widgets/Client/SClientName.h b/Engine/Plugins/Developer/Concert/ConcertUI/ConcertSharedSlate/Source/ConcertSharedSlate/Public/Widgets/Client/SClientName.h
new file mode 100644
--- /dev/null
+++ b/Engine/Plugins/Developer/Concert/ConcertUI/ConcertSharedSlate/Source/ConcertSharedSlate/Public/Widgets/Client/SClientName.h
@@ -0,0 +1,55 @@
+// Copyright Epic Games, Inc. All Rights Reserved.
+
+#pragma once
+
+#include "ConcertMessageData.h"
+#include "Styling/AppStyle.h"
+#include "Widgets/DeclarativeSyntaxSupport.h"
+#include "Widgets/SCompoundWidget.h"
+
+class IConcertClient;
+
+namespace UE::ConcertClientSharedSlate
+{
+	/** Knows how to display FConcertClientInfo. */
+	class CONCERTSHAREDSLATE_API SClientName : public SCompoundWidget
+	{
+	public:
+
+		SLATE_BEGIN_ARGS(SClientName)
+			: _DisplayAsLocalClient(false)
+			, _DisplayAvatarColor(true)
+			, _Font(FAppStyle::Get().GetFontStyle("BoldFont"))
+		{}
+			/** The client info to display. */
+			SLATE_ATTRIBUTE(TOptional<FConcertClientInfo>, ClientInfo)
+			
+			/** Whether visually indicate that this is a local client (appends "(me)" if true). */
+			SLATE_ATTRIBUTE(bool, DisplayAsLocalClient)
+			/** Whether to show a square image in front of the name. */
+			SLATE_ATTRIBUTE(bool, DisplayAvatarColor)
+			
+			/** Used for highlighting in the text */
+			SLATE_ATTRIBUTE(FText, HighlightText)
+			/** The font to use for the name */
+			SLATE_ARGUMENT(FSlateFontInfo, Font)
+		SLATE_END_ARGS()
+
+		void Construct(const FArguments& InArgs);
+
+		/** @return The display that would be used given the settings. */
+		static FText GetDisplayText(const FConcertClientInfo& Info, bool bDisplayAsLocalClient);
+
+	private:
+
+		/** The client info to display. */
+		TAttribute<TOptional<FConcertClientInfo>> ClientInfoAttribute;
+		/** Whether visually indicate that this is a local client (appends "(me)" if true). */
+		TAttribute<bool> DisplayAsLocalClientAttribute;
+
+		/** Gets the display name. */
+		FText GetClientDisplayName() const;
+		/** Gets the avatar color */
+		FSlateColor GetAvatarColor() const;
+	};
+}
diff --git a/Engine/Plugins/Developer/Concert/ConcertUI/ConcertSharedSlate/Source/ConcertSharedSlate/Public/Widgets/Client/SHorizontalClientList.h b/Engine/Plugins/Developer/Concert/ConcertUI/ConcertSharedSlate/Source/ConcertSharedSlate/Public/Widgets/Client/SHorizontalClientList.h
new file mode 100644
--- /dev/null
+++ b/Engine/Plugins/Developer/Concert/ConcertUI/ConcertSharedSlate/Source/ConcertSharedSlate/Public/Widgets/Client/SHorizontalClientList.h
@@ -0,0 +1,85 @@
+// Copyright Epic Games, Inc. All Rights Reserved.
+
+#pragma once
+
+#include "ClientInfoDelegate.h"
+#include "Styling/AppStyle.h"
+#include "Widgets/DeclarativeSyntaxSupport.h"
+#include "Widgets/SCompoundWidget.h"
+
+class IConcertClient;
+class SScrollBox;
+class SWidgetSwitcher;
+
+struct FConcertSessionClientInfo;
+struct FGuid;
+
+namespace UE::ConcertClientSharedSlate
+{
+	/** Aligns client widgets from left to right. If there is not enough space, a horizontal scroll bar cuts of the list. */
+	class CONCERTSHAREDSLATE_API SHorizontalClientList : public SCompoundWidget
+	{
+	public:
+
+		DECLARE_DELEGATE_RetVal_TwoParams(bool, FSortPredicate, const FConcertSessionClientInfo& Left, const FConcertSessionClientInfo& FConcertSessionClientInfo);
+		static bool SortLocalClientFirstThenAlphabetical(const FConcertSessionClientInfo& Left, const FConcertSessionClientInfo& Right, ConcertSharedSlate::FIsLocalClient IsLocalClientDelegate);
+
+		/** @return The display string a SHorizontalClientList would display with the given state. Returns unset optional if EmptyListSlot would be shown. */
+		static TOptional<FString> GetDisplayString(
+			const TConstArrayView<FGuid>& Clients,
+			const ConcertSharedSlate::FGetOptionalClientInfo& GetClientInfoDelegate,
+			const FSortPredicate& SortPredicate,
+			const ConcertSharedSlate::FIsLocalClient& IsLocalClientDelegate = {}
+			);
+		
+		SLATE_BEGIN_ARGS(SHorizontalClientList)
+			: _Font(FAppStyle::Get().GetFontStyle("NormalFont"))
+		{}
+			/** Decides whether the given client should be displayed as if it was a local client. */
+			SLATE_EVENT(ConcertSharedSlate::FIsLocalClient, IsLocalClient)
+			/** Used to get client display info for remote clients. */
+			SLATE_EVENT(ConcertSharedSlate::FGetOptionalClientInfo, GetClientInfo)
+			
+			/** Whether to show a square image in front of the name. */
+			SLATE_ATTRIBUTE(bool, DisplayAvatarColor)
+			
+			/** Used for highlighting in the text */
+			SLATE_ATTRIBUTE(FText, HighlightText)
+			/** The font to use for the names */
+			SLATE_ARGUMENT(FSlateFontInfo, Font)
+			
+			/** Defaults to placing the local client first (if contained) and sorting alphabetically otherwise. */
+			SLATE_EVENT(FSortPredicate, SortPredicate)
+			
+			/** The widget to display when the list is empty */
+			SLATE_NAMED_SLOT(FArguments, EmptyListSlot)
+		SLATE_END_ARGS()
+
+		void Construct(const FArguments& InArgs);
+
+		/** Refreshes the list. */
+		void RefreshList(const TConstArrayView<FGuid>& Clients);
+
+	private:
+		
+		/** Decides whether the given client should be displayed as if it was a local client. */
+		ConcertSharedSlate::FIsLocalClient IsLocalClientDelegate;
+		/** Used to get client display info for remote clients. */
+		ConcertSharedSlate::FGetOptionalClientInfo GetClientInfoDelegate;
+		/** Sorts the client list */
+		FSortPredicate SortPredicateDelegate;
+
+		/** Whether the square in front of the client name should be displayed. */
+		TAttribute<bool> DisplayAvatarColorAttribute;
+		/** Used for highlighting in the text */
+		TAttribute<FText> HighlightTextAttribute;
+		
+		/** The font to use for the names */
+		FSlateFontInfo NameFont;
+
+		/** Displays the ScrollBox when there are clients and the EmptyListSlot otherwise. */
+		TSharedPtr<SWidgetSwitcher> WidgetSwitcher;
+		/** Contains the children. */
+		TSharedPtr<SScrollBox> ScrollBox;
+	};
+}
\ No newline at end of file
diff --git a/Engine/Plugins/Developer/Concert/ConcertUI/ConcertSharedSlate/Source/ConcertSharedSlate/Public/Widgets/Client/SLocalClientName.h b/Engine/Plugins/Developer/Concert/ConcertUI/ConcertSharedSlate/Source/ConcertSharedSlate/Public/Widgets/Client/SLocalClientName.h
new file mode 100644
--- /dev/null
+++ b/Engine/Plugins/Developer/Concert/ConcertUI/ConcertSharedSlate/Source/ConcertSharedSlate/Public/Widgets/Client/SLocalClientName.h
@@ -0,0 +1,42 @@
+// Copyright Epic Games, Inc. All Rights Reserved.
+
+#pragma once
+
+#include "ConcertMessageData.h"
+#include "Misc/Optional.h"
+#include "Styling/AppStyle.h"
+#include "Widgets/DeclarativeSyntaxSupport.h"
+#include "Widgets/SCompoundWidget.h"
+
+class IConcertClient;
+
+namespace UE::ConcertClientSharedSlate
+{
+	/**
+	 * Displays the name of a client.
+	 * 
+	 * The name will look like "Client Name (me)".
+	 * @see SRemoteClientName
+	 */
+	class CONCERTSHAREDSLATE_API SLocalClientName : public SCompoundWidget
+	{
+	public:
+
+		SLATE_BEGIN_ARGS(SLocalClientName)
+			: _Font(FAppStyle::Get().GetFontStyle("BoldFont"))
+		{}
+			/** The client info to display. */
+			SLATE_ATTRIBUTE(TOptional<FConcertClientInfo>, DisplayInfo)
+			
+			/** Whether to show a square image in front of the name. */
+			SLATE_ATTRIBUTE(bool, DisplayAvatarColor)
+			
+			/** Used for highlighting in the text */
+			SLATE_ATTRIBUTE(FText, HighlightText)
+			/** The font to use for the name */
+			SLATE_ARGUMENT(FSlateFontInfo, Font)
+		SLATE_END_ARGS()
+
+		void Construct(const FArguments& InArgs);
+	};
+}
diff --git a/Engine/Plugins/Developer/Concert/ConcertUI/ConcertSharedSlate/Source/ConcertSharedSlate/Public/Widgets/Client/SRemoteClientName.h b/Engine/Plugins/Developer/Concert/ConcertUI/ConcertSharedSlate/Source/ConcertSharedSlate/Public/Widgets/Client/SRemoteClientName.h
new file mode 100644
--- /dev/null
+++ b/Engine/Plugins/Developer/Concert/ConcertUI/ConcertSharedSlate/Source/ConcertSharedSlate/Public/Widgets/Client/SRemoteClientName.h
@@ -0,0 +1,58 @@
+// Copyright Epic Games, Inc. All Rights Reserved.
+
+#pragma once
+
+#include "ConcertMessageData.h"
+#include "Misc/Optional.h"
+#include "Styling/AppStyle.h"
+#include "Widgets/DeclarativeSyntaxSupport.h"
+#include "Widgets/SCompoundWidget.h"
+
+class IConcertClient;
+
+namespace UE::ConcertClientSharedSlate
+{
+	/**
+	 * Displays the name of a client.
+	 * 
+	 * The name will look like "Client Name".
+	 * @see SLocalClientName
+	 *
+	 * If the client disconnects, the last known info is used.
+	 * If the client info is unknown, the widget will display an empty FConcertClientInfo;
+	 */
+	class CONCERTSHAREDSLATE_API SRemoteClientName : public SCompoundWidget
+	{
+	public:
+
+		SLATE_BEGIN_ARGS(SRemoteClientName)
+			: _Font(FAppStyle::Get().GetFontStyle("BoldFont"))
+		{}
+			/** The client info to display. */
+			SLATE_ATTRIBUTE(TOptional<FConcertClientInfo>, DisplayInfo)
+			
+			/** Whether to show a square image in front of the name. */
+			SLATE_ATTRIBUTE(bool, DisplayAvatarColor)
+			
+			/** Used for highlighting in the text */
+			SLATE_ATTRIBUTE(FText, HighlightText)
+			/** The font to use for the name */
+			SLATE_ARGUMENT(FSlateFontInfo, Font)
+		SLATE_END_ARGS()
+
+		void Construct(const FArguments& InArgs);
+
+	private:
+
+		/** The endpoint ID of the client to display. */
+		TAttribute<TOptional<FConcertClientInfo>> ClientDisplayInfo;
+
+		/**
+		 * Cached so that the info remains known when the client disconnects.
+		 * Must be mutable because TAttribute::CreateSP requires GetClientInfo to be const.
+		 */
+		mutable TOptional<FConcertClientInfo> LastKnownClientInfo;
+
+		TOptional<FConcertClientInfo> GetClientInfo() const;
+	};
+}
