Remove incorrect #undef LOCTEXT_NAMESPACE

diff --git a/Engine/Plugins/Developer/Concert/ConcertUI/ConcertSharedSlate/Source/ConcertSharedSlate/Private/Replication/Editor/View/Tree/TreeItemTraits.h b/Engine/Plugins/Developer/Concert/ConcertUI/ConcertSharedSlate/Source/ConcertSharedSlate/Private/Replication/Editor/View/Tree/TreeItemTraits.h
--- a/Engine/Plugins/Developer/Concert/ConcertUI/ConcertSharedSlate/Source/ConcertSharedSlate/Private/Replication/Editor/View/Tree/TreeItemTraits.h
+++ b/Engine/Plugins/Developer/Concert/ConcertUI/ConcertSharedSlate/Source/ConcertSharedSlate/Private/Replication/Editor/View/Tree/TreeItemTraits.h
@@ -50,6 +50,4 @@ namespace UE::ConcertSharedSlate
 			return IsLessThan_Default(Left, Right, SortByColumn);
 		}
 	};
-}
-
-#undef LOCTEXT_NAMESPACE
\ No newline at end of file
+}
\ No newline at end of file
