Rename UI labels in Multi User:
- "Overview" tab label was renamed to "Session" (MU tab when connected to session)
- "Add objects" combo button label was renamed to "Add" (MU tab when connected to session > replication tab)

diff --git a/Engine/Plugins/Developer/Concert/ConcertApp/MultiUserClient/Source/MultiUserClient/Private/Widgets/ActiveSession/SActiveSessionRoot.cpp b/Engine/Plugins/Developer/Concert/ConcertApp/MultiUserClient/Source/MultiUserClient/Private/Widgets/ActiveSession/SActiveSessionRoot.cpp
--- a/Engine/Plugins/Developer/Concert/ConcertApp/MultiUserClient/Source/MultiUserClient/Private/Widgets/ActiveSession/SActiveSessionRoot.cpp
+++ b/Engine/Plugins/Developer/Concert/ConcertApp/MultiUserClient/Source/MultiUserClient/Private/Widgets/ActiveSession/SActiveSessionRoot.cpp
@@ -16,7 +16,7 @@
 
 namespace UE::MultiUserClient
 {
-	const FName SActiveSessionRoot::OverviewTabId(TEXT("OverviewTabId"));
+	const FName SActiveSessionRoot::SessionOverviewTabId(TEXT("OverviewTabId"));
 	const FName SActiveSessionRoot::ReplicationTabId(TEXT("ReplicationTabId"));
 	
 	void SActiveSessionRoot::Construct(
@@ -42,9 +42,9 @@ namespace UE::MultiUserClient
 				->Split
 				(
 					FTabManager::NewStack()
-					->AddTab(OverviewTabId, ETabState::OpenedTab)
+					->AddTab(SessionOverviewTabId, ETabState::OpenedTab)
 					->AddTab(ReplicationTabId, ETabState::OpenedTab)
-					->SetForegroundTab(OverviewTabId)
+					->SetForegroundTab(SessionOverviewTabId)
 					->SetHideTabWell(false)
 					)
 			);
@@ -82,7 +82,7 @@ namespace UE::MultiUserClient
 		TSharedRef<FMultiUserReplicationManager> InReplicationManager
 		)
 	{
-		InTabManager->RegisterTabSpawner(OverviewTabId, FOnSpawnTab::CreateSP(this, &SActiveSessionRoot::SpawnTab_Overview))
+		InTabManager->RegisterTabSpawner(SessionOverviewTabId, FOnSpawnTab::CreateSP(this, &SActiveSessionRoot::SpawnTab_SessionOverview))
 			.SetDisplayName(LOCTEXT("OverviewTab.DisplayName", "Overview"))
 			.SetGroup(AppMenuGroup);
 	
@@ -91,11 +91,11 @@ namespace UE::MultiUserClient
 			.SetGroup(AppMenuGroup);
 	}
 
-	TSharedRef<SDockTab> SActiveSessionRoot::SpawnTab_Overview(const FSpawnTabArgs& Args)
+	TSharedRef<SDockTab> SActiveSessionRoot::SpawnTab_SessionOverview(const FSpawnTabArgs& Args)
 	{
 		return SNew(SDockTab)
-			.Label(LOCTEXT("OverviewTab.Label", "Overview"))
-			.ToolTipText(LOCTEXT("OverviewTab.Tooltip", "Displays active session clients and activity."))
+			.Label(LOCTEXT("SessionOverviewTab.Label", "Session"))
+			.ToolTipText(LOCTEXT("SessionOverviewTab.Tooltip", "Displays active session clients and activity."))
 			[
 				SNew(SActiveSessionOverviewTab, ConcertSyncClient)
 			];
diff --git a/Engine/Plugins/Developer/Concert/ConcertApp/MultiUserClient/Source/MultiUserClient/Private/Widgets/ActiveSession/SActiveSessionRoot.h b/Engine/Plugins/Developer/Concert/ConcertApp/MultiUserClient/Source/MultiUserClient/Private/Widgets/ActiveSession/SActiveSessionRoot.h
--- a/Engine/Plugins/Developer/Concert/ConcertApp/MultiUserClient/Source/MultiUserClient/Private/Widgets/ActiveSession/SActiveSessionRoot.h
+++ b/Engine/Plugins/Developer/Concert/ConcertApp/MultiUserClient/Source/MultiUserClient/Private/Widgets/ActiveSession/SActiveSessionRoot.h
@@ -24,7 +24,7 @@ namespace UE::MultiUserClient
 	{
 	public:
 		
-		static const FName OverviewTabId;
+		static const FName SessionOverviewTabId;
 		static const FName ReplicationTabId;
 
 		SLATE_BEGIN_ARGS(SActiveSessionRoot)
@@ -40,7 +40,7 @@ namespace UE::MultiUserClient
 		TSharedPtr<FTabManager> TabManager;
 		
 		void RegisterTabSpawners(const TSharedRef<FTabManager>& InTabManager, const TSharedRef<FWorkspaceItem>& AppMenuGroup, TSharedRef<FMultiUserReplicationManager> InReplicationManager);
-		TSharedRef<SDockTab> SpawnTab_Overview(const FSpawnTabArgs& Args);
+		TSharedRef<SDockTab> SpawnTab_SessionOverview(const FSpawnTabArgs& Args);
 		TSharedRef<SDockTab> SpawnTab_ReplicationControls(const FSpawnTabArgs& Args, TSharedRef<FMultiUserReplicationManager> InReplicationManager);
 	};
 }
diff --git a/Engine/Plugins/Developer/Concert/ConcertUI/ConcertClientSharedSlate/Source/ConcertClientSharedSlate/Private/Replication/Editor/Model/ObjectSource/ActorSelectionSourceModel.cpp b/Engine/Plugins/Developer/Concert/ConcertUI/ConcertClientSharedSlate/Source/ConcertClientSharedSlate/Private/Replication/Editor/Model/ObjectSource/ActorSelectionSourceModel.cpp
--- a/Engine/Plugins/Developer/Concert/ConcertUI/ConcertClientSharedSlate/Source/ConcertClientSharedSlate/Private/Replication/Editor/Model/ObjectSource/ActorSelectionSourceModel.cpp
+++ b/Engine/Plugins/Developer/Concert/ConcertUI/ConcertClientSharedSlate/Source/ConcertClientSharedSlate/Private/Replication/Editor/Model/ObjectSource/ActorSelectionSourceModel.cpp
@@ -18,8 +18,8 @@ namespace UE::ConcertClientSharedSlate
 			{
 				ConcertSharedSlate::FBaseDisplayInfo
 				{
-					LOCTEXT("ObjectsCategory.Label", "Add Objects"),
-					LOCTEXT("ObjectsCategory.Tooltip", "Options for adding objects from the open editor world")
+					LOCTEXT("ObjectsCategory.Label", "Add"),
+					LOCTEXT("ObjectsCategory.Tooltip", "Add objects from the open editor world")
 				},
 				{
 					MakeShared<FSelectedActorsSource>(),
diff --git a/Engine/Plugins/Developer/Concert/ConcertUI/ConcertSharedSlate/Source/ConcertSharedSlate/Public/Replication/ReplicationWidgetFactories.h b/Engine/Plugins/Developer/Concert/ConcertUI/ConcertSharedSlate/Source/ConcertSharedSlate/Public/Replication/ReplicationWidgetFactories.h
--- a/Engine/Plugins/Developer/Concert/ConcertUI/ConcertSharedSlate/Source/ConcertSharedSlate/Public/Replication/ReplicationWidgetFactories.h
+++ b/Engine/Plugins/Developer/Concert/ConcertUI/ConcertSharedSlate/Source/ConcertSharedSlate/Public/Replication/ReplicationWidgetFactories.h
@@ -55,11 +55,8 @@ namespace UE::ConcertSharedSlate
 	/** Params for creating a IReplicationStreamViewer. */
 	struct FCreateViewerParams
 	{
-		/** Required. Displays the properties in a tree view. You can pass in e.g. custom UI with advanced filtering. */
-		TSharedRef<IPropertyTreeView> PropertyTreeView = CreateSearchablePropertyTreeView();
-
 		/** Required. In the lower half of the editor, this view presents the properties associated with the object that is currently selected in the upper part of the view. */
-		TSharedRef<IPropertyAssignmentView> PropertyAssignmentView = CreatePerObjectAssignmentView(); // TODO DP: Change this to the matrix view
+		TSharedRef<IPropertyAssignmentView> PropertyAssignmentView = CreatePerObjectAssignmentView();
 		
 		/**
 		 * Optional. Determines the objects displayed as children to the top-level objects in the top section.
@@ -88,8 +85,6 @@ namespace UE::ConcertSharedSlate
 		TAlwaysValidWidget RightOfObjectSearchBar;
 	};
 
-	// TODO DP 5.5: Create factory function that uses FCreateViewerParams and creates an IReplicationStreamViewer
-
 	/** Params for creating an IReplicationStreamEditor */
 	struct FCreateEditorParams
 	{
@@ -144,9 +139,9 @@ namespace UE::ConcertSharedSlate
 		/**
 		 * This model consolidates all objects in all streams.
 		 * 
-		 * When an object is added
-		 * - to one of the streams,
-		 * - via the "add" button,
+		 * When an object is added (i.e. IEditableReplicationStreamModel::OnObjectsChanged broadcasts), e.g. when
+		 * - to an object is added to one of the managed streams (e.g. stream is remotely changed),
+		 * - via the "Add" button,
 		 * it is added here, too. 
 		 * This is needed for the internal operation of the multi view.
 		 *
@@ -168,7 +163,7 @@ namespace UE::ConcertSharedSlate
 		 */
 		TSharedRef<IPropertySelectionSourceModel> PropertySource;
 
-		/** Optional. If set, the Add Actor button should automatically assign the added object to stream returned */
+		/** Optional. If set, the Add button should automatically assign the added object to stream returned by this callback. */
 		FGetAutoAssignTarget GetAutoAssignToStreamDelegate;
 	};
 
