In the Multi User replication tab, only shows objects that are in the client's current world.
Objects from other worlds (added by remote clients) are not shown.

diff --git a/Engine/Plugins/Developer/Concert/ConcertApp/MultiUserClient/Source/MultiUserClient/Private/Replication/Client/ReplicationClient.cpp b/Engine/Plugins/Developer/Concert/ConcertApp/MultiUserClient/Source/MultiUserClient/Private/Replication/Client/ReplicationClient.cpp
--- a/Engine/Plugins/Developer/Concert/ConcertApp/MultiUserClient/Source/MultiUserClient/Private/Replication/Client/ReplicationClient.cpp
+++ b/Engine/Plugins/Developer/Concert/ConcertApp/MultiUserClient/Source/MultiUserClient/Private/Replication/Client/ReplicationClient.cpp
@@ -36,7 +36,9 @@ namespace UE::MultiUserClient
 			SubmissionQueue
 			)
 		, LocalClientEditModel(
+			// Transact ClientContentStorage
 			ConcertClientSharedSlate::CreateTransactionalStreamModel(
+				// Read & write the stream data in ClientContentStorage
 				CreateBaseStreamModel(
 					ClientContentStorage->Stream->MakeReplicationMapGetterAttribute(),
 					MakeShared<FMultiUserStreamExtender>(EndpointId, InDiscoveryContainer)
diff --git a/Engine/Plugins/Developer/Concert/ConcertApp/MultiUserClient/Source/MultiUserClient/Private/Widgets/ActiveSession/Replication/Client/Multi/SMultiClientView.cpp b/Engine/Plugins/Developer/Concert/ConcertApp/MultiUserClient/Source/MultiUserClient/Private/Widgets/ActiveSession/Replication/Client/Multi/SMultiClientView.cpp
--- a/Engine/Plugins/Developer/Concert/ConcertApp/MultiUserClient/Source/MultiUserClient/Private/Widgets/ActiveSession/Replication/Client/Multi/SMultiClientView.cpp
+++ b/Engine/Plugins/Developer/Concert/ConcertApp/MultiUserClient/Source/MultiUserClient/Private/Widgets/ActiveSession/Replication/Client/Multi/SMultiClientView.cpp
@@ -52,6 +52,12 @@ namespace UE::MultiUserClient
 			);
 
 		RebuildClientSubscriptions();
+
+		// Changing worlds affects what things are displayed in the editor.
+		HideObjectsNotInEditorWorld.OnRefreshObjects().AddLambda([this]()
+		{
+			StreamEditor->GetEditorBase().Refresh();
+		});
 	}
 
 	SMultiClientView::~SMultiClientView()
@@ -112,7 +118,8 @@ namespace UE::MultiUserClient
 			{
 				MultiStreamColumns::ReplicationToggle(InConcertClient, ObjectHierarchyAttribute, InClientManager),
 				MultiStreamColumns::AssignedClientsColumn(InConcertClient, MultiStreamEditorAttribute, ObjectHierarchyAttribute, InClientManager.GetReassignmentLogic(), InClientManager)
-			}
+			},
+			.ShouldDisplayObjectDelegate = FShouldDisplayObject::CreateSP(this, &SMultiClientView::ShouldDisplayObject)
 		};
 		
 		StreamEditor = CreateBaseMultiStreamEditor(MoveTemp(Params), MoveTemp(ViewerParams));
@@ -200,6 +207,11 @@ namespace UE::MultiUserClient
 				);
 		}
 	}
+
+	bool SMultiClientView::ShouldDisplayObject(const FSoftObjectPath& Object) const
+	{
+		return HideObjectsNotInEditorWorld.ShouldShowObject(Object);
+	}
 }
 
 #undef LOCTEXT_NAMESPACE
\ No newline at end of file
diff --git a/Engine/Plugins/Developer/Concert/ConcertApp/MultiUserClient/Source/MultiUserClient/Private/Widgets/ActiveSession/Replication/Client/Multi/SMultiClientView.h b/Engine/Plugins/Developer/Concert/ConcertApp/MultiUserClient/Source/MultiUserClient/Private/Widgets/ActiveSession/Replication/Client/Multi/SMultiClientView.h
--- a/Engine/Plugins/Developer/Concert/ConcertApp/MultiUserClient/Source/MultiUserClient/Private/Widgets/ActiveSession/Replication/Client/Multi/SMultiClientView.h
+++ b/Engine/Plugins/Developer/Concert/ConcertApp/MultiUserClient/Source/MultiUserClient/Private/Widgets/ActiveSession/Replication/Client/Multi/SMultiClientView.h
@@ -2,6 +2,7 @@
 
 #pragma once
 
+#include "Replication/Editor/UnrealEditor/HideObjectsNotInWorldLogic.h"
 #include "Widgets/DeclarativeSyntaxSupport.h"
 #include "Widgets/SCompoundWidget.h"
 
@@ -46,6 +47,9 @@ namespace UE::MultiUserClient
 		/** Used by widgets in columns. */
 		TSharedPtr<ConcertSharedSlate::IObjectHierarchyModel> ObjectHierarchy;
 
+		/** This logic helps us decide whether an object should be displayed and lets us know that the object list needs to be refreshed (e.g. due to world change). */
+		ConcertClientSharedSlate::FHideObjectsNotInWorldLogic HideObjectsNotInEditorWorld;
+
 		/** Creates this widget's editor content */
 		TSharedRef<SWidget> CreateEditorContent(const TSharedRef<IConcertClient>& InConcertClient, FReplicationClientManager& InClientManager);
 
@@ -60,5 +64,7 @@ namespace UE::MultiUserClient
 		
 		/** Adds additional entries to the context menu for the object tree view. */
 		void ExtendObjectContextMenu(FMenuBuilder& MenuBuilder, TConstArrayView<FSoftObjectPath> ContextObjects) const;
+		/** Decides whether the object should be displayed: do not show it if it's not in the editor world. */
+		bool ShouldDisplayObject(const FSoftObjectPath& Object) const;
 	};
 }
diff --git a/Engine/Plugins/Developer/Concert/ConcertUI/ConcertClientSharedSlate/Source/ConcertClientSharedSlate/Private/Replication/ClientReplicationWidgetFactories.cpp b/Engine/Plugins/Developer/Concert/ConcertUI/ConcertClientSharedSlate/Source/ConcertClientSharedSlate/Private/Replication/ClientReplicationWidgetFactories.cpp
--- a/Engine/Plugins/Developer/Concert/ConcertUI/ConcertClientSharedSlate/Source/ConcertClientSharedSlate/Private/Replication/ClientReplicationWidgetFactories.cpp
+++ b/Engine/Plugins/Developer/Concert/ConcertUI/ConcertClientSharedSlate/Source/ConcertClientSharedSlate/Private/Replication/ClientReplicationWidgetFactories.cpp
@@ -2,12 +2,12 @@
 
 #include "Replication/ClientReplicationWidgetFactories.h"
 
-#include "Editor/View/PropertyTree/SFilteredPropertyTreeView.h"
 #include "Replication/ReplicationWidgetFactories.h"
 #include "Replication/Editor/Model/Object/EditorObjectHierarchyModel.h"
 #include "Replication/Editor/Model/Object/EditorObjectNameModel.h"
 #include "Replication/Editor/Model/ReplicationStreamObject.h"
 #include "Replication/Editor/Model/TransactionalReplicationStreamModel.h"
+#include "Replication/Editor/View/PropertyTree/SFilteredPropertyTreeView.h"
 
 #include "UObject/UObjectGlobals.h"
 #include "UObject/Package.h"
@@ -25,12 +25,12 @@ namespace UE::ConcertClientSharedSlate
 	}
 
 	TSharedRef<ConcertSharedSlate::IEditableReplicationStreamModel> CreateTransactionalStreamModel(
-		TSharedRef<ConcertSharedSlate::IEditableReplicationStreamModel> BaseModel,
+		const TSharedRef<ConcertSharedSlate::IEditableReplicationStreamModel>& BaseModel,
 		UObject& OwnerObject
 		)
 	{
 		return MakeShared<ConcertSharedSlate::FTransactionalReplicationStreamModel>(
-			MoveTemp(BaseModel),
+			BaseModel,
 			OwnerObject
 			);
 	}
diff --git a/Engine/Plugins/Developer/Concert/ConcertUI/ConcertClientSharedSlate/Source/ConcertClientSharedSlate/Private/Replication/Editor/Model/TransactionalReplicationStreamModel.cpp b/Engine/Plugins/Developer/Concert/ConcertUI/ConcertClientSharedSlate/Source/ConcertClientSharedSlate/Private/Replication/Editor/Model/TransactionalReplicationStreamModel.cpp
--- a/Engine/Plugins/Developer/Concert/ConcertUI/ConcertClientSharedSlate/Source/ConcertClientSharedSlate/Private/Replication/Editor/Model/TransactionalReplicationStreamModel.cpp
+++ b/Engine/Plugins/Developer/Concert/ConcertUI/ConcertClientSharedSlate/Source/ConcertClientSharedSlate/Private/Replication/Editor/Model/TransactionalReplicationStreamModel.cpp
@@ -10,10 +10,10 @@
 namespace UE::ConcertSharedSlate
 {
 	FTransactionalReplicationStreamModel::FTransactionalReplicationStreamModel(
-		TSharedRef<IEditableReplicationStreamModel> WrappedModel,
+		const TSharedRef<IEditableReplicationStreamModel>& WrappedModel,
 		UObject& OwningObject
 		)
-		: WrappedModel(MoveTemp(WrappedModel))
+		: FEditableStreamModelProxy(WrappedModel)
 		, OwningObject(&OwningObject)
 	{}
 
@@ -21,28 +21,28 @@ namespace UE::ConcertSharedSlate
 	{
 		const FScopedTransaction Transaction(LOCTEXT("AddObjects", "Add replicated objects"));
 		OwningObject->Modify();
-		WrappedModel->AddObjects(Objects);
+		GetWrappedModel()->AddObjects(Objects);
 	}
 
 	void FTransactionalReplicationStreamModel::RemoveObjects(TConstArrayView<FSoftObjectPath> Objects)
 	{
 		const FScopedTransaction Transaction(LOCTEXT("RemoveObjects", "Remove replicated objects"));
 		OwningObject->Modify();
-		WrappedModel->RemoveObjects(Objects);
+		GetWrappedModel()->RemoveObjects(Objects);
 	}
 
 	void FTransactionalReplicationStreamModel::AddProperties(const FSoftObjectPath& SoftObjectPath, TConstArrayView<FConcertPropertyChain> Properties)
 	{
 		const FScopedTransaction Transaction(LOCTEXT("AddProperties", "Add replicated properties"));
 		OwningObject->Modify();
-		WrappedModel->AddProperties(SoftObjectPath, Properties);
+		GetWrappedModel()->AddProperties(SoftObjectPath, Properties);
 	}
 
 	void FTransactionalReplicationStreamModel::RemoveProperties(const FSoftObjectPath& SoftObjectPath, TConstArrayView<FConcertPropertyChain> Properties)
 	{
 		const FScopedTransaction Transaction(LOCTEXT("RemoveProperties", "Remove replicated properties"));
 		OwningObject->Modify();
-		WrappedModel->RemoveProperties(SoftObjectPath, Properties);
+		GetWrappedModel()->RemoveProperties(SoftObjectPath, Properties);
 	}
 
 	bool FTransactionalReplicationStreamModel::MatchesContext(const FTransactionContext& InContext, const TArray<TPair<UObject*, FTransactionObjectEvent>>& TransactionObjectContexts) const
diff --git a/Engine/Plugins/Developer/Concert/ConcertUI/ConcertClientSharedSlate/Source/ConcertClientSharedSlate/Private/Replication/Editor/Model/TransactionalReplicationStreamModel.h b/Engine/Plugins/Developer/Concert/ConcertUI/ConcertClientSharedSlate/Source/ConcertClientSharedSlate/Private/Replication/Editor/Model/TransactionalReplicationStreamModel.h
--- a/Engine/Plugins/Developer/Concert/ConcertUI/ConcertClientSharedSlate/Source/ConcertClientSharedSlate/Private/Replication/Editor/Model/TransactionalReplicationStreamModel.h
+++ b/Engine/Plugins/Developer/Concert/ConcertUI/ConcertClientSharedSlate/Source/ConcertClientSharedSlate/Private/Replication/Editor/Model/TransactionalReplicationStreamModel.h
@@ -2,8 +2,9 @@
 
 #pragma once
 
+#include "Replication/Editor/Model/Util/EditableStreamModelProxy.h"
+
 #include "EditorUndoClient.h"
-#include "Replication/Editor/Model/IEditableReplicationStreamModel.h"
 #include "UObject/GCObject.h"
 
 struct FConcertObjectReplicationMap;
@@ -14,27 +15,16 @@ namespace UE::ConcertSharedSlate
 
 	/** Special case of FGenericPropertySelectionModel where the edited FConcertObjectReplicationMap lives in an UObject that is RF_Transactional. */
 	class FTransactionalReplicationStreamModel
-		: public IEditableReplicationStreamModel
+		: public FEditableStreamModelProxy
 		, public FSelfRegisteringEditorUndoClient
 		, public FGCObject
 	{
 	public:
 
 		FTransactionalReplicationStreamModel(
-			TSharedRef<IEditableReplicationStreamModel> WrappedModel,
+			const TSharedRef<IEditableReplicationStreamModel>& WrappedModel,
 			UObject& OwningObject
 			);
-
-		//~ Begin IReplicationStreamModel Interface
-		virtual FSoftClassPath GetObjectClass(const FSoftObjectPath& Object) const override { return WrappedModel->GetObjectClass(Object); }
-		virtual bool ContainsObjects(const TSet<FSoftObjectPath>& Objects) const override { return WrappedModel->ContainsObjects(Objects); }
-		virtual bool ContainsProperties(const FSoftObjectPath& Object, const TSet<FConcertPropertyChain>& Properties) const override { return WrappedModel->ContainsProperties(Object, Properties); }
-		virtual bool ForEachReplicatedObject(TFunctionRef<EBreakBehavior(const FSoftObjectPath& Object)> Delegate) const override { return WrappedModel->ForEachReplicatedObject(Delegate); }
-		virtual bool ForEachProperty(const FSoftObjectPath& Object, TFunctionRef<EBreakBehavior(const FConcertPropertyChain& Property)> Delegate) const override { return WrappedModel->ForEachProperty(Object, Delegate); }
-		virtual uint32 GetNumProperties(const FSoftObjectPath& Object) const override { return WrappedModel->GetNumProperties(Object); }
-		virtual FOnObjectsChanged& OnObjectsChanged() override { return WrappedModel->OnObjectsChanged(); }
-		virtual FOnPropertiesChanged& OnPropertiesChanged() override { return WrappedModel->OnPropertiesChanged(); }
-		//~ End IReplicationStreamModel Interface
 		
 		//~ Begin IEditableReplicationStreamModel Interface
 		virtual void AddObjects(TConstArrayView<UObject*> Objects) override;
@@ -55,9 +45,6 @@ namespace UE::ConcertSharedSlate
 		//~ End FGCObject Interface
 
 	private:
-
-		/** The model to be transacted */
-		TSharedRef<IEditableReplicationStreamModel> WrappedModel;
 		
 		/** User of FTransactionalPropertySelectionModel is responsible for keeping OwningObject alive, e.g. via an asset editor. */
 		TWeakObjectPtr<UObject> OwningObject;
diff --git a/Engine/Plugins/Developer/Concert/ConcertUI/ConcertClientSharedSlate/Source/ConcertClientSharedSlate/Private/Replication/Editor/UnrealEditor/HideObjectsNotInWorldLogic.cpp b/Engine/Plugins/Developer/Concert/ConcertUI/ConcertClientSharedSlate/Source/ConcertClientSharedSlate/Private/Replication/Editor/UnrealEditor/HideObjectsNotInWorldLogic.cpp
new file mode 100644
--- /dev/null
+++ b/Engine/Plugins/Developer/Concert/ConcertUI/ConcertClientSharedSlate/Source/ConcertClientSharedSlate/Private/Replication/Editor/UnrealEditor/HideObjectsNotInWorldLogic.cpp
@@ -0,0 +1,40 @@
+﻿// Copyright Epic Games, Inc. All Rights Reserved.
+
+#include "Replication/Editor/UnrealEditor/HideObjectsNotInWorldLogic.h"
+
+#include "Engine/Engine.h"
+#include "Engine/World.h"
+#include "UObject/SoftObjectPath.h"
+
+namespace UE::ConcertClientSharedSlate
+{
+	FHideObjectsNotInWorldLogic::FHideObjectsNotInWorldLogic()
+	{
+		if (ensure(GEngine))
+		{
+			GEngine->OnWorldAdded().AddRaw(this, &FHideObjectsNotInWorldLogic::OnWorldAdded);
+			GEngine->OnWorldDestroyed().AddRaw(this, &FHideObjectsNotInWorldLogic::OnWorldDestroyed);
+		}
+	}
+
+	FHideObjectsNotInWorldLogic::~FHideObjectsNotInWorldLogic()
+	{
+		if (GEngine)
+		{
+			GEngine->OnWorldAdded().RemoveAll(this);
+			GEngine->OnWorldDestroyed().RemoveAll(this);
+		}
+	}
+
+	bool FHideObjectsNotInWorldLogic::ShouldShowObject(const FSoftObjectPath& ObjectPath) const
+	{
+		if (!GWorld)
+		{
+			return false;
+		}
+
+		// If it does not resolve, it's no in GWorld (the world should load all the actors in it).
+		const UObject* Object = ObjectPath.ResolveObject();
+		return Object && Object->IsIn(GWorld);
+	}
+}
diff --git a/Engine/Plugins/Developer/Concert/ConcertUI/ConcertClientSharedSlate/Source/ConcertClientSharedSlate/Public/Replication/ClientReplicationWidgetFactories.h b/Engine/Plugins/Developer/Concert/ConcertUI/ConcertClientSharedSlate/Source/ConcertClientSharedSlate/Public/Replication/ClientReplicationWidgetFactories.h
--- a/Engine/Plugins/Developer/Concert/ConcertUI/ConcertClientSharedSlate/Source/ConcertClientSharedSlate/Public/Replication/ClientReplicationWidgetFactories.h
+++ b/Engine/Plugins/Developer/Concert/ConcertUI/ConcertClientSharedSlate/Source/ConcertClientSharedSlate/Public/Replication/ClientReplicationWidgetFactories.h
@@ -37,10 +37,9 @@ namespace UE::ConcertClientSharedSlate
 	 * @param OwnerObject The object containing the FConcertObjectReplicationMap - used for transactions.
 	 */
 	CONCERTCLIENTSHAREDSLATE_API TSharedRef<ConcertSharedSlate::IEditableReplicationStreamModel> CreateTransactionalStreamModel(
-		TSharedRef<ConcertSharedSlate::IEditableReplicationStreamModel> BaseModel,
+		const TSharedRef<ConcertSharedSlate::IEditableReplicationStreamModel>& BaseModel,
 		UObject& OwnerObject
 		);
-
 	/** Simpler CreateTransactionalStreamModel overload that internally creates an UObject and sets it up automatically. */
 	CONCERTCLIENTSHAREDSLATE_API TSharedRef<ConcertSharedSlate::IEditableReplicationStreamModel> CreateTransactionalStreamModel();
 
diff --git a/Engine/Plugins/Developer/Concert/ConcertUI/ConcertClientSharedSlate/Source/ConcertClientSharedSlate/Public/Replication/Editor/UnrealEditor/HideObjectsNotInWorldLogic.h b/Engine/Plugins/Developer/Concert/ConcertUI/ConcertClientSharedSlate/Source/ConcertClientSharedSlate/Public/Replication/Editor/UnrealEditor/HideObjectsNotInWorldLogic.h
new file mode 100644
--- /dev/null
+++ b/Engine/Plugins/Developer/Concert/ConcertUI/ConcertClientSharedSlate/Source/ConcertClientSharedSlate/Public/Replication/Editor/UnrealEditor/HideObjectsNotInWorldLogic.h
@@ -0,0 +1,39 @@
+﻿// Copyright Epic Games, Inc. All Rights Reserved.
+
+#pragma once
+
+#include "Delegates/Delegate.h"
+#include "Templates/UnrealTemplate.h"
+
+class UWorld;
+struct FSoftObjectPath;
+
+namespace UE::ConcertClientSharedSlate
+{
+	/**
+	 * This logic is used to show only those replicated objects that are in the world the user has opened in the editor.
+	 * Its functions are injected into the FCreateViewerParams.
+	 */
+	class CONCERTCLIENTSHAREDSLATE_API FHideObjectsNotInWorldLogic : public FNoncopyable
+	{
+	public:
+		
+		FHideObjectsNotInWorldLogic();
+		~FHideObjectsNotInWorldLogic();
+		
+		/** @return Whether this object should be shown to the user. */
+		bool ShouldShowObject(const FSoftObjectPath& ObjectPath) const;
+
+		DECLARE_MULTICAST_DELEGATE(FOnRefreshObjects);
+		FOnRefreshObjects& OnRefreshObjects() { return OnRefreshObjectsDelegate; }
+		
+	private:
+
+		/** Broadcasts when the set of displayed objects has changed, e.g. when switching worlds. */
+		FOnRefreshObjects OnRefreshObjectsDelegate;
+		
+		void OnWorldAdded(UWorld* World) const { OnRefreshObjectsDelegate.Broadcast(); }
+		void OnWorldDestroyed(UWorld* World) const { OnRefreshObjectsDelegate.Broadcast(); }
+	};
+}
+
diff --git a/Engine/Plugins/Developer/Concert/ConcertUI/ConcertSharedSlate/Source/ConcertSharedSlate/Private/Replication/Editor/View/MultiEditor/ConsolidatedMultiStreamModel.h b/Engine/Plugins/Developer/Concert/ConcertUI/ConcertSharedSlate/Source/ConcertSharedSlate/Private/Replication/Editor/View/MultiEditor/ConsolidatedMultiStreamModel.h
--- a/Engine/Plugins/Developer/Concert/ConcertUI/ConcertSharedSlate/Source/ConcertSharedSlate/Private/Replication/Editor/View/MultiEditor/ConsolidatedMultiStreamModel.h
+++ b/Engine/Plugins/Developer/Concert/ConcertUI/ConcertSharedSlate/Source/ConcertSharedSlate/Private/Replication/Editor/View/MultiEditor/ConsolidatedMultiStreamModel.h
@@ -36,7 +36,7 @@ namespace UE::ConcertSharedSlate
 			TSharedRef<IEditableMultiReplicationStreamModel> InMultiStreamModel,
 			FGetAutoAssignTarget InGetAutoAssignTargetDelegate
 			);
-		virtual ~FConsolidatedMultiStreamModel();
+		virtual ~FConsolidatedMultiStreamModel() override;
 		
 		//~ Begin IReplicationStreamModel Interface
 		virtual FSoftClassPath GetObjectClass(const FSoftObjectPath& Object) const override;
diff --git a/Engine/Plugins/Developer/Concert/ConcertUI/ConcertSharedSlate/Source/ConcertSharedSlate/Private/Replication/Editor/View/ObjectEditor/SBaseReplicationStreamEditor.cpp b/Engine/Plugins/Developer/Concert/ConcertUI/ConcertSharedSlate/Source/ConcertSharedSlate/Private/Replication/Editor/View/ObjectEditor/SBaseReplicationStreamEditor.cpp
--- a/Engine/Plugins/Developer/Concert/ConcertUI/ConcertSharedSlate/Source/ConcertSharedSlate/Private/Replication/Editor/View/ObjectEditor/SBaseReplicationStreamEditor.cpp
+++ b/Engine/Plugins/Developer/Concert/ConcertUI/ConcertSharedSlate/Source/ConcertSharedSlate/Private/Replication/Editor/View/ObjectEditor/SBaseReplicationStreamEditor.cpp
@@ -21,15 +21,15 @@ namespace UE::ConcertSharedSlate
 {
 	void SBaseReplicationStreamEditor::Construct(
 		const FArguments& InArgs,
-		TSharedRef<IEditableReplicationStreamModel> InPropertiesModel,
-		TSharedRef<IObjectSelectionSourceModel> InObjectSelectionSource,
-		TSharedRef<IPropertySelectionSourceModel> InPropertySelectionSource)
+		const TSharedRef<IEditableReplicationStreamModel>& InPropertiesModel,
+		const TSharedRef<IObjectSelectionSourceModel>& InObjectSelectionSource,
+		const TSharedRef<IPropertySelectionSourceModel>& InPropertySelectionSource)
 	{
-		ObjectSelectionSource = MoveTemp(InObjectSelectionSource);
-		PropertySelectionSource = MoveTemp(InPropertySelectionSource);
+		ObjectSelectionSource = InObjectSelectionSource;
+		PropertySelectionSource = InPropertySelectionSource;
 		ObjectHierarchy = InArgs._ObjectHierarchy;
 		
-		EditablePropertiesModel = MoveTemp(InPropertiesModel);
+		EditablePropertiesModel = InPropertiesModel;
 		EditablePropertiesModel->OnObjectsChanged().AddSP(this, &SBaseReplicationStreamEditor::OnObjectsChanged);
 		EditablePropertiesModel->OnPropertiesChanged().AddSP(this, &SBaseReplicationStreamEditor::OnPropertiesChanged);
 
@@ -48,6 +48,7 @@ namespace UE::ConcertSharedSlate
 				.NameModel(InArgs._NameModel)
 				.OnDeleteObjects(this, &SBaseReplicationStreamEditor::OnDeleteObjects)
 				.OnObjectsContextMenuOpening(this, &SBaseReplicationStreamEditor::OnObjectsContextMenuOpening)
+				.ShouldDisplayObject(InArgs._ShouldDisplayObject)
 				.LeftOfObjectSearchBar()
 				[
 					SNew(SHorizontalBox)
diff --git a/Engine/Plugins/Developer/Concert/ConcertUI/ConcertSharedSlate/Source/ConcertSharedSlate/Private/Replication/Editor/View/ObjectEditor/SBaseReplicationStreamEditor.h b/Engine/Plugins/Developer/Concert/ConcertUI/ConcertSharedSlate/Source/ConcertSharedSlate/Private/Replication/Editor/View/ObjectEditor/SBaseReplicationStreamEditor.h
--- a/Engine/Plugins/Developer/Concert/ConcertUI/ConcertSharedSlate/Source/ConcertSharedSlate/Private/Replication/Editor/View/ObjectEditor/SBaseReplicationStreamEditor.h
+++ b/Engine/Plugins/Developer/Concert/ConcertUI/ConcertSharedSlate/Source/ConcertSharedSlate/Private/Replication/Editor/View/ObjectEditor/SBaseReplicationStreamEditor.h
@@ -69,6 +69,9 @@ namespace UE::ConcertSharedSlate
 
 			/** Called to generate the context menu for objects. This extends the options already generated by this widget. */
 			SLATE_EVENT(FExtendObjectMenu, OnExtendObjectsContextMenu)
+			
+			/** Optional. Whether a given object should be displayed. If this returns false on an object, none of its children will be shown either. */
+			SLATE_EVENT(FShouldDisplayObject, ShouldDisplayObject)
 
 			/** Optional widget to add to the left of the object list search bar. */
 			SLATE_NAMED_SLOT(FArguments, LeftOfObjectSearchBar)
@@ -82,9 +85,9 @@ namespace UE::ConcertSharedSlate
 		SLATE_END_ARGS()
 
 		void Construct(const FArguments& InArgs,
-           TSharedRef<IEditableReplicationStreamModel> InPropertiesModel,
-           TSharedRef<IObjectSelectionSourceModel> InObjectSelectionSource,
-           TSharedRef<IPropertySelectionSourceModel> InPropertySelectionSource
+           const TSharedRef<IEditableReplicationStreamModel>& InPropertiesModel,
+           const TSharedRef<IObjectSelectionSourceModel>& InObjectSelectionSource,
+           const TSharedRef<IPropertySelectionSourceModel>& InPropertySelectionSource
 		);
 		virtual ~SBaseReplicationStreamEditor() override;
 
diff --git a/Engine/Plugins/Developer/Concert/ConcertUI/ConcertSharedSlate/Source/ConcertSharedSlate/Private/Replication/Editor/View/ObjectViewer/SReplicationStreamViewer.cpp b/Engine/Plugins/Developer/Concert/ConcertUI/ConcertSharedSlate/Source/ConcertSharedSlate/Private/Replication/Editor/View/ObjectViewer/SReplicationStreamViewer.cpp
--- a/Engine/Plugins/Developer/Concert/ConcertUI/ConcertSharedSlate/Source/ConcertSharedSlate/Private/Replication/Editor/View/ObjectViewer/SReplicationStreamViewer.cpp
+++ b/Engine/Plugins/Developer/Concert/ConcertUI/ConcertSharedSlate/Source/ConcertSharedSlate/Private/Replication/Editor/View/ObjectViewer/SReplicationStreamViewer.cpp
@@ -25,11 +25,12 @@
 
 namespace UE::ConcertSharedSlate
 {
-	void SReplicationStreamViewer::Construct(const FArguments& InArgs, TSharedRef<IReplicationStreamModel> InPropertiesModel)
+	void SReplicationStreamViewer::Construct(const FArguments& InArgs, const TSharedRef<IReplicationStreamModel>& InPropertiesModel)
 	{
-		PropertiesModel = MoveTemp(InPropertiesModel);
+		PropertiesModel = InPropertiesModel;
 		ObjectHierarchy = InArgs._ObjectHierarchy;
 		NameModel = InArgs._NameModel;
+		ShouldDisplayObjectDelegate = InArgs._ShouldDisplayObject;
 		
 		ChildSlot
 		[
@@ -308,7 +309,7 @@ namespace UE::ConcertSharedSlate
 		{
 			TOptional<IObjectHierarchyModel::FParentInfo> ParentInfo = ObjectHierarchy->GetParentInfo(ObjectPath);
 			const bool bIsActor = !ParentInfo; 
-			if (bIsActor || ShouldDisplayObject(ObjectPath, ParentInfo->Relationship))
+			if (bIsActor || ShouldDisplayChildObject(ObjectPath, ParentInfo->Relationship))
 			{
 				const TSharedPtr<FReplicatedObjectData>* ExistingItem = PathToObjectDataCache.Find(ObjectPath);
 				ExistingItem = ExistingItem ? ExistingItem : NewPathToObjectDataCache.Find(ObjectPath);
@@ -346,6 +347,11 @@ namespace UE::ConcertSharedSlate
 		EBreakBehavior BreakBehavior = EBreakBehavior::Continue;
 		PropertiesModel->ForEachReplicatedObject([this, &Delegate, &NonActors, &Actors, &BreakBehavior](const FSoftObjectPath& Object)
 		{
+			if (!CanDisplayObject(Object))
+			{
+				return EBreakBehavior::Continue;
+			}
+			
 			Delegate(Object);
 			
 			if (ObjectUtils::IsActor(Object))
@@ -370,7 +376,7 @@ namespace UE::ConcertSharedSlate
 		for (const FSoftObjectPath& NonTopLevelObject : NonActors)
 		{
 			const TOptional<FSoftObjectPath> TopLevelObject = ObjectUtils::GetActorOf(NonTopLevelObject);
-			if (!TopLevelObject)
+			if (!TopLevelObject || !CanDisplayObject(*TopLevelObject))
 			{
 				continue;
 			}
@@ -446,7 +452,7 @@ namespace UE::ConcertSharedSlate
 		AddItem(OwningActor);
 		ObjectHierarchy->ForEachChildRecursive(OwningActor, [this, &AddItem](const FSoftObjectPath&, const FSoftObjectPath& ChildObject, EChildRelationship Relationship)
 		{
-			if (ShouldDisplayObject(ChildObject, Relationship))
+			if (ShouldDisplayChildObject(ChildObject, Relationship))
 			{
 				AddItem(ChildObject);
 			}
@@ -473,8 +479,12 @@ namespace UE::ConcertSharedSlate
 			return EBreakBehavior::Continue;
 		});
 	}
+	bool SReplicationStreamViewer::CanDisplayObject(const FSoftObjectPath& ObjectPath) const
+	{
+		return !ShouldDisplayObjectDelegate.IsBound() || ShouldDisplayObjectDelegate.Execute(ObjectPath);
+	}
 
-	bool SReplicationStreamViewer::ShouldDisplayObject(const FSoftObjectPath& Object, EChildRelationship Relationship) const
+	bool SReplicationStreamViewer::ShouldDisplayObjectRelation(EChildRelationship Relationship) const
 	{
 		const bool bSkipSubobject = Relationship == EChildRelationship::Subobject && !ObjectViewOptions.ShouldDisplaySubobjects();
 		return !bSkipSubobject;
diff --git a/Engine/Plugins/Developer/Concert/ConcertUI/ConcertSharedSlate/Source/ConcertSharedSlate/Private/Replication/Editor/View/ObjectViewer/SReplicationStreamViewer.h b/Engine/Plugins/Developer/Concert/ConcertUI/ConcertSharedSlate/Source/ConcertSharedSlate/Private/Replication/Editor/View/ObjectViewer/SReplicationStreamViewer.h
--- a/Engine/Plugins/Developer/Concert/ConcertUI/ConcertSharedSlate/Source/ConcertSharedSlate/Private/Replication/Editor/View/ObjectViewer/SReplicationStreamViewer.h
+++ b/Engine/Plugins/Developer/Concert/ConcertUI/ConcertSharedSlate/Source/ConcertSharedSlate/Private/Replication/Editor/View/ObjectViewer/SReplicationStreamViewer.h
@@ -72,6 +72,9 @@ namespace UE::ConcertSharedSlate
 		
 			/** Called to generate the context menu for objects. */
 			SLATE_EVENT(FOnContextMenuOpening, OnObjectsContextMenuOpening)
+			
+			/** Optional. Whether a given object should be displayed. If this returns false on an object, none of its children will be shown either. */
+			SLATE_EVENT(FShouldDisplayObject, ShouldDisplayObject)
 		
 			/** Optional widget to add to the left of the object list search bar. */
 			SLATE_NAMED_SLOT(FArguments, LeftOfObjectSearchBar)
@@ -82,7 +85,7 @@ namespace UE::ConcertSharedSlate
 			SLATE_ATTRIBUTE(FText, NoOutlinerObjects)
 		SLATE_END_ARGS()
 
-		void Construct(const FArguments& InArgs, TSharedRef<IReplicationStreamModel> InPropertiesModel);
+		void Construct(const FArguments& InArgs, const TSharedRef<IReplicationStreamModel>& InPropertiesModel);
 
 		//~ Begin IReplicationStreamViewer Interface
 		virtual void Refresh() override;
@@ -138,7 +141,10 @@ namespace UE::ConcertSharedSlate
 		bool bHasRequestedPropertyRefresh = false;
 		TArray<FSoftObjectPath> PendingToSelect;
 		TArray<FSoftObjectPath> PendingToExpand;
-		bool bPendingExpandRecursively;
+		bool bPendingExpandRecursively = false;
+		
+		/** Optional. Whether a given object should be displayed. If this returns false on an object, none of its children will be shown either. */
+		FShouldDisplayObject ShouldDisplayObjectDelegate;
 
 		static TSharedRef<FReplicatedObjectData> AllocateObjectData(FSoftObjectPath ObjectPath);
 
@@ -172,7 +178,15 @@ namespace UE::ConcertSharedSlate
 			RequestObjectDataRefresh();
 			RequestPropertyDataRefresh();
 		}
-		bool ShouldDisplayObject(const FSoftObjectPath& Object, EChildRelationship Relationship) const;
+		bool ShouldDisplayChildObject(const FSoftObjectPath& Object, EChildRelationship Relationship) const
+		{
+			return CanDisplayObject(Object) && ShouldDisplayObjectRelation(Relationship);
+		}
+		
+		/** Invokes the ShouldDisplayObjectDelegate to determine whether ObjectPath should be displayed. */
+		bool CanDisplayObject(const FSoftObjectPath& ObjectPath) const;
+		/** Whether the view options allow this type of relationship to be shown. */
+		bool ShouldDisplayObjectRelation(EChildRelationship Relationship) const;
 		
 		FSoftClassPath GetObjectClass(const FSoftObjectPath& ObjectPath) const;
 	};
diff --git a/Engine/Plugins/Developer/Concert/ConcertUI/ConcertSharedSlate/Source/ConcertSharedSlate/Private/Replication/ReplicationWidgetFactories.cpp b/Engine/Plugins/Developer/Concert/ConcertUI/ConcertSharedSlate/Source/ConcertSharedSlate/Private/Replication/ReplicationWidgetFactories.cpp
--- a/Engine/Plugins/Developer/Concert/ConcertUI/ConcertSharedSlate/Source/ConcertSharedSlate/Private/Replication/ReplicationWidgetFactories.cpp
+++ b/Engine/Plugins/Developer/Concert/ConcertUI/ConcertSharedSlate/Source/ConcertSharedSlate/Private/Replication/ReplicationWidgetFactories.cpp
@@ -22,7 +22,7 @@ namespace UE::ConcertSharedSlate
 	
 	TSharedRef<IReplicationStreamEditor> CreateBaseStreamEditor(FCreateEditorParams EditorParams, FCreateViewerParams ViewerParams)
 	{
-		return SNew(SBaseReplicationStreamEditor, MoveTemp(EditorParams.DataModel), MoveTemp(EditorParams.ObjectSource), MoveTemp(EditorParams.PropertySource))
+		return SNew(SBaseReplicationStreamEditor, EditorParams.DataModel, EditorParams.ObjectSource, EditorParams.PropertySource)
 			.PropertyAssignmentView(MoveTemp(ViewerParams.PropertyAssignmentView))
 			.ObjectColumns(MoveTemp(ViewerParams.ObjectColumns))
 			.PrimaryObjectSort(ViewerParams.PrimaryObjectSort)
@@ -30,6 +30,7 @@ namespace UE::ConcertSharedSlate
 			.ObjectHierarchy(MoveTemp(ViewerParams.ObjectHierarchy))
 			.NameModel(MoveTemp(ViewerParams.NameModel))
 			.OnExtendObjectsContextMenu(MoveTemp(ViewerParams.OnExtendObjectsContextMenu))
+			.ShouldDisplayObject(MoveTemp(ViewerParams.ShouldDisplayObjectDelegate))
 			.LeftOfObjectSearchBar() [ MoveTemp(ViewerParams.LeftOfObjectSearchBar.Widget) ]
 			.RightOfObjectSearchBar() [ MoveTemp(ViewerParams.RightOfObjectSearchBar.Widget) ]
 			.IsEditingEnabled(MoveTemp(EditorParams.IsEditingEnabled))
diff --git a/Engine/Plugins/Developer/Concert/ConcertUI/ConcertSharedSlate/Source/ConcertSharedSlate/Public/Replication/Editor/Model/Util/EditableStreamModelProxy.h b/Engine/Plugins/Developer/Concert/ConcertUI/ConcertSharedSlate/Source/ConcertSharedSlate/Public/Replication/Editor/Model/Util/EditableStreamModelProxy.h
new file mode 100644
--- /dev/null
+++ b/Engine/Plugins/Developer/Concert/ConcertUI/ConcertSharedSlate/Source/ConcertSharedSlate/Public/Replication/Editor/Model/Util/EditableStreamModelProxy.h
@@ -0,0 +1,48 @@
+﻿// Copyright Epic Games, Inc. All Rights Reserved.
+
+#pragma once
+
+#include "Replication/Editor/Model/IEditableReplicationStreamModel.h"
+
+namespace UE::ConcertSharedSlate
+{
+	/**
+	 * Base util class for implementations that alter or extend the behavior of another IEditableReplicationStreamModel.
+	 * This pattern allows chaining of behaviors.
+	 */
+	class FEditableStreamModelProxy : public IEditableReplicationStreamModel
+	{
+	public:
+
+		FEditableStreamModelProxy(const TSharedRef<IEditableReplicationStreamModel>& InModel)
+			: WrappedModel(InModel)
+		{}
+
+		//~ Begin IReplicationStreamModel Interface
+		virtual FSoftClassPath GetObjectClass(const FSoftObjectPath& Object) const override { return WrappedModel->GetObjectClass(Object); }
+		virtual bool ContainsObjects(const TSet<FSoftObjectPath>& Objects) const override { return WrappedModel->ContainsObjects(Objects); }
+		virtual bool ContainsProperties(const FSoftObjectPath& Object, const TSet<FConcertPropertyChain>& Properties) const override { return WrappedModel->ContainsProperties(Object, Properties); }
+		virtual bool ForEachReplicatedObject(TFunctionRef<EBreakBehavior(const FSoftObjectPath& Object)> Delegate) const override { return WrappedModel->ForEachReplicatedObject(Delegate); }
+		virtual bool ForEachProperty(const FSoftObjectPath& Object, TFunctionRef<EBreakBehavior(const FConcertPropertyChain& Property)> Delegate) const override { return WrappedModel->ForEachProperty(Object, Delegate); }
+		virtual uint32 GetNumProperties(const FSoftObjectPath& Object) const override { return WrappedModel->GetNumProperties(Object); }
+		virtual FOnObjectsChanged& OnObjectsChanged() override { return WrappedModel->OnObjectsChanged(); }
+		virtual FOnPropertiesChanged& OnPropertiesChanged() override { return WrappedModel->OnPropertiesChanged(); }
+		//~ End IReplicationStreamModel Interface
+		
+		//~ Begin IEditableReplicationStreamModel Interface
+		virtual void AddObjects(TConstArrayView<UObject*> Objects) override { WrappedModel->AddObjects(Objects); }
+		virtual void RemoveObjects(TConstArrayView<FSoftObjectPath> Objects) override { WrappedModel->RemoveObjects(Objects); }
+		virtual void AddProperties(const FSoftObjectPath& Object, TConstArrayView<FConcertPropertyChain> Properties) override { WrappedModel->AddProperties(Object, Properties); }
+		virtual void RemoveProperties(const FSoftObjectPath& Object, TConstArrayView<FConcertPropertyChain> Properties) override { WrappedModel->RemoveProperties(Object, Properties); }
+		//~ End IEditableReplicationStreamModel Interface
+
+	protected:
+
+		const TSharedRef<IEditableReplicationStreamModel>& GetWrappedModel() const { return WrappedModel; }
+
+	private:
+		
+		/** The model to be transacted */
+		TSharedRef<IEditableReplicationStreamModel> WrappedModel;
+	};
+}
diff --git a/Engine/Plugins/Developer/Concert/ConcertUI/ConcertSharedSlate/Source/ConcertSharedSlate/Public/Replication/ReplicationWidgetFactories.h b/Engine/Plugins/Developer/Concert/ConcertUI/ConcertSharedSlate/Source/ConcertSharedSlate/Public/Replication/ReplicationWidgetFactories.h
--- a/Engine/Plugins/Developer/Concert/ConcertUI/ConcertSharedSlate/Source/ConcertSharedSlate/Public/Replication/ReplicationWidgetFactories.h
+++ b/Engine/Plugins/Developer/Concert/ConcertUI/ConcertSharedSlate/Source/ConcertSharedSlate/Public/Replication/ReplicationWidgetFactories.h
@@ -6,7 +6,6 @@
 #include "Editor/View/Column/IObjectTreeColumn.h"
 #include "Editor/View/Column/ReplicationColumnInfo.h"
 #include "Editor/View/Column/SelectionViewerColumns.h"
-#include "PropertyTreeFactory.h"
 #include "Replication/Editor/View/Column/ReplicationColumnsUtils.h"
 #include "Replication/Utils/ReplicationWidgetDelegates.h"
 
@@ -83,6 +82,17 @@ namespace UE::ConcertSharedSlate
 		TAlwaysValidWidget LeftOfObjectSearchBar;
 		/** Optional widget to add to the right of the object list search bar. */
 		TAlwaysValidWidget RightOfObjectSearchBar;
+		
+		/**
+		 * Optional. Whether a given object should be displayed. 
+		 *
+		 * This is useful e.g. for hiding objects that are not in the local editor's opened world (for that purpuse couple this delegate with @see FHideObjectsNotInWorldLogic).
+		 * If the list of displayed objects changes, call IReplicationStreamViewer::Refresh.
+		 * 
+		 * If this returns false on an object, none of its children will be shown either
+		 * Child objects are determined using ObjectHierarchy.
+		 */
+		FShouldDisplayObject ShouldDisplayObjectDelegate;
 	};
 
 	/** Params for creating an IReplicationStreamEditor */
diff --git a/Engine/Plugins/Developer/Concert/ConcertUI/ConcertSharedSlate/Source/ConcertSharedSlate/Public/Replication/Utils/ReplicationWidgetDelegates.h b/Engine/Plugins/Developer/Concert/ConcertUI/ConcertSharedSlate/Source/ConcertSharedSlate/Public/Replication/Utils/ReplicationWidgetDelegates.h
--- a/Engine/Plugins/Developer/Concert/ConcertUI/ConcertSharedSlate/Source/ConcertSharedSlate/Public/Replication/Utils/ReplicationWidgetDelegates.h
+++ b/Engine/Plugins/Developer/Concert/ConcertUI/ConcertSharedSlate/Source/ConcertSharedSlate/Public/Replication/Utils/ReplicationWidgetDelegates.h
@@ -24,6 +24,10 @@ namespace UE::ConcertSharedSlate
 
 	/** Extends a context menu that is being built for a selection of objects. */
 	DECLARE_DELEGATE_TwoParams(FExtendObjectMenu, FMenuBuilder&, TConstArrayView<FSoftObjectPath> ContextObjects);
-	
+
+	/** Delegate for getting an object's class. */
 	DECLARE_DELEGATE_RetVal_OneParam(FSoftClassPath, FGetObjectClass, const FSoftObjectPath&);
+
+	/** Delegate for deciding whether an object should be displayed. */
+	DECLARE_DELEGATE_RetVal_OneParam(bool, FShouldDisplayObject, const FSoftObjectPath& ObjectPath);
 }
\ No newline at end of file
