Improve performance of updating property list.
- Sorting properties now uses a cache to speed to generating names of FProperties.
- UI now uses TSoftObjectPath, which cases objects after resolve, instead of loading it every time

diff --git a/Engine/Plugins/Developer/Concert/ConcertSync/ConcertSyncCore/Source/ConcertSyncCore/Private/Replication/Data/ConcertPropertySelection.cpp b/Engine/Plugins/Developer/Concert/ConcertSync/ConcertSyncCore/Source/ConcertSyncCore/Private/Replication/Data/ConcertPropertySelection.cpp
--- a/Engine/Plugins/Developer/Concert/ConcertSync/ConcertSyncCore/Source/ConcertSyncCore/Private/Replication/Data/ConcertPropertySelection.cpp
+++ b/Engine/Plugins/Developer/Concert/ConcertSync/ConcertSyncCore/Source/ConcertSyncCore/Private/Replication/Data/ConcertPropertySelection.cpp
@@ -163,7 +163,7 @@ bool FConcertPropertyChain::MatchesExactly(const FArchiveSerializedPropertyChain
     return bArePathsEqual && bVisitedEveryProperty;
 }
 
-FProperty* FConcertPropertyChain::ResolveProperty(UStruct& Class, bool bLogOnFail) const
+FProperty* FConcertPropertyChain::ResolveProperty(const UStruct& Class, bool bLogOnFail) const
 {
     // FConcertPropertyChain::ResolveProperty exists only for visibility to developers since
     // the class is the first place one would look and not in the utils namespace.
diff --git a/Engine/Plugins/Developer/Concert/ConcertSync/ConcertSyncCore/Source/ConcertSyncCore/Private/Replication/PropertyChainUtils.cpp b/Engine/Plugins/Developer/Concert/ConcertSync/ConcertSyncCore/Source/ConcertSyncCore/Private/Replication/PropertyChainUtils.cpp
--- a/Engine/Plugins/Developer/Concert/ConcertSync/ConcertSyncCore/Source/ConcertSyncCore/Private/Replication/PropertyChainUtils.cpp
+++ b/Engine/Plugins/Developer/Concert/ConcertSync/ConcertSyncCore/Source/ConcertSyncCore/Private/Replication/PropertyChainUtils.cpp
@@ -114,6 +114,8 @@ namespace UE::ConcertSyncCore::PropertyChain
 	
 	FProperty* ResolveProperty(const UStruct& Class, const FConcertPropertyChain& ChainToResolve, bool bLogOnFail)
 	{
+		SCOPED_CONCERT_TRACE(ResolveProperty);
+		
 		FProperty* Result = nullptr;
 		ForEachReplicatableProperty(
 			Class,
diff --git a/Engine/Plugins/Developer/Concert/ConcertSync/ConcertSyncCore/Source/ConcertSyncCore/Private/Replication/PropertyResolutionCache.cpp b/Engine/Plugins/Developer/Concert/ConcertSync/ConcertSyncCore/Source/ConcertSyncCore/Private/Replication/PropertyResolutionCache.cpp
new file mode 100644
--- /dev/null
+++ b/Engine/Plugins/Developer/Concert/ConcertSync/ConcertSyncCore/Source/ConcertSyncCore/Private/Replication/PropertyResolutionCache.cpp
@@ -0,0 +1,34 @@
+﻿// Copyright Epic Games, Inc. All Rights Reserved.
+
+#include "Replication/PropertyResolutionCache.h"
+
+#include "Replication/Data/ConcertPropertySelection.h"
+
+#include "UObject/Class.h"
+
+namespace UE::ConcertSyncCore::PropertyChain
+{
+	FProperty* FPropertyResolutionCache::ResolveAndCache(const UStruct& Struct, const FConcertPropertyChain& Chain)
+	{
+		FClassCache& ClassCache = CachedProperties.FindOrAdd(&Struct);
+		if (FProperty** const CachedProperty = ClassCache.Cache.Find(Chain))
+		{
+			return *CachedProperty;
+		}
+
+		FProperty* ResolvedProperty = Chain.ResolveProperty(Struct);
+		UClass* OwningClass = ResolvedProperty ? ResolvedProperty->GetOwnerClass() : nullptr;
+		UStruct* OwningStruct = ResolvedProperty ? ResolvedProperty->GetOwnerStruct() : nullptr;
+		UScriptStruct* OwningScriptStruct = Cast<UScriptStruct>(OwningStruct);
+
+		const bool bIsNativeClass = OwningClass && OwningClass->HasAnyClassFlags(CLASS_Native);
+		const bool bIsNativeStruct = (OwningScriptStruct && OwningScriptStruct->StructFlags & STRUCT_Native)
+			// Example: FVector and PRESUMABLY other classes in UObject/NoExportTypes.h
+			|| (OwningStruct && !OwningScriptStruct);
+		if (bIsNativeClass || bIsNativeStruct)
+		{
+			ClassCache.Cache.Add(Chain, ResolvedProperty);
+		}
+		return ResolvedProperty;
+	}
+}
\ No newline at end of file
diff --git a/Engine/Plugins/Developer/Concert/ConcertSync/ConcertSyncCore/Source/ConcertSyncCore/Public/Replication/Data/ConcertPropertySelection.h b/Engine/Plugins/Developer/Concert/ConcertSync/ConcertSyncCore/Source/ConcertSyncCore/Public/Replication/Data/ConcertPropertySelection.h
--- a/Engine/Plugins/Developer/Concert/ConcertSync/ConcertSyncCore/Source/ConcertSyncCore/Public/Replication/Data/ConcertPropertySelection.h
+++ b/Engine/Plugins/Developer/Concert/ConcertSync/ConcertSyncCore/Source/ConcertSyncCore/Public/Replication/Data/ConcertPropertySelection.h
@@ -50,7 +50,7 @@ struct CONCERTSYNCCORE_API FConcertPropertyChain
 	bool MatchesExactly(const FArchiveSerializedPropertyChain* OptionalChain, const FProperty& LeafProperty) const;
 
 	/** @return Attempts to resolve this property given the class */
-	FProperty* ResolveProperty(UStruct& Class, bool bLogOnFail = true) const;
+	FProperty* ResolveProperty(const UStruct& Class, bool bLogOnFail = true) const;
 
 	FConcertPropertyChain GetParent() const;
 	const TArray<FName>& GetPathToProperty() const { return PathToProperty; }
diff --git a/Engine/Plugins/Developer/Concert/ConcertSync/ConcertSyncCore/Source/ConcertSyncCore/Public/Replication/PropertyResolutionCache.h b/Engine/Plugins/Developer/Concert/ConcertSync/ConcertSyncCore/Source/ConcertSyncCore/Public/Replication/PropertyResolutionCache.h
new file mode 100644
--- /dev/null
+++ b/Engine/Plugins/Developer/Concert/ConcertSync/ConcertSyncCore/Source/ConcertSyncCore/Public/Replication/PropertyResolutionCache.h
@@ -0,0 +1,39 @@
+﻿// Copyright Epic Games, Inc. All Rights Reserved.
+
+#pragma once
+
+#include "Containers/Map.h"
+#include "Data/ConcertPropertySelection.h"
+
+class UStruct;
+
+namespace UE::ConcertSyncCore::PropertyChain
+{
+	/**
+	 * FConcertPropertyChain::ResolveProperty is expensive because it has to iterate the property hierarchy.
+	 * 
+	 * This caches look-up results of native classes.
+	 * By default, properties in Blueprint classes are not cached because they can change at runtime and leave dangling FProperty pointers.
+	 */
+	class CONCERTSYNCCORE_API FPropertyResolutionCache
+	{
+	public:
+
+		/** Resolves the property chain and caches the result; */
+		FProperty* ResolveAndCache(const UStruct& Struct, const FConcertPropertyChain& Chain);
+
+		/** Removes the class from the cache.*/
+		void Invalidate(const UStruct& Struct) { CachedProperties.Remove(&Struct); }
+		/** Removes everything from the cache. */
+		void Clear(int32 ExpectedNumElements = 0) { CachedProperties.Empty(ExpectedNumElements); }
+
+	private:
+
+		struct FClassCache
+		{
+			TMap<FConcertPropertyChain, FProperty*> Cache;
+		};
+
+		TMap<const UStruct*, FClassCache> CachedProperties;
+	};
+}
diff --git a/Engine/Plugins/Developer/Concert/ConcertUI/ConcertClientSharedSlate/Source/ConcertClientSharedSlate/Private/Replication/Editor/Model/Object/EditorObjectNameModel.cpp b/Engine/Plugins/Developer/Concert/ConcertUI/ConcertClientSharedSlate/Source/ConcertClientSharedSlate/Private/Replication/Editor/Model/Object/EditorObjectNameModel.cpp
--- a/Engine/Plugins/Developer/Concert/ConcertUI/ConcertClientSharedSlate/Source/ConcertClientSharedSlate/Private/Replication/Editor/Model/Object/EditorObjectNameModel.cpp
+++ b/Engine/Plugins/Developer/Concert/ConcertUI/ConcertClientSharedSlate/Source/ConcertClientSharedSlate/Private/Replication/Editor/Model/Object/EditorObjectNameModel.cpp
@@ -32,9 +32,9 @@ namespace UE::ConcertClientSharedSlate
 		}
 	}
 	
-	FText FEditorObjectNameModel::GetObjectDisplayName(const FSoftObjectPath& ObjectPath) const
+	FText FEditorObjectNameModel::GetObjectDisplayName(const TSoftObjectPtr<>& ObjectPath) const
 	{
-		if (UObject* ResolvedObject = ObjectPath.ResolveObject())
+		if (UObject* ResolvedObject = ObjectPath.Get())
 		{
 			// Display actor just like the outliner does
 			if (const AActor* AsActor = Cast<AActor>(ResolvedObject))
@@ -53,6 +53,6 @@ namespace UE::ConcertClientSharedSlate
 			return FText::FromString(ResolvedObject->GetName());
 		}
 		
-		return FText::FromString(ConcertSharedSlate::ObjectUtils::ExtractObjectDisplayStringFromPath(ObjectPath));
+		return FText::FromString(ConcertSharedSlate::ObjectUtils::ExtractObjectDisplayStringFromPath(ObjectPath.GetUniqueID()));
 	}
 }
diff --git a/Engine/Plugins/Developer/Concert/ConcertUI/ConcertClientSharedSlate/Source/ConcertClientSharedSlate/Private/Replication/Editor/Model/Object/EditorObjectNameModel.h b/Engine/Plugins/Developer/Concert/ConcertUI/ConcertClientSharedSlate/Source/ConcertClientSharedSlate/Private/Replication/Editor/Model/Object/EditorObjectNameModel.h
--- a/Engine/Plugins/Developer/Concert/ConcertUI/ConcertClientSharedSlate/Source/ConcertClientSharedSlate/Private/Replication/Editor/Model/Object/EditorObjectNameModel.h
+++ b/Engine/Plugins/Developer/Concert/ConcertUI/ConcertClientSharedSlate/Source/ConcertClientSharedSlate/Private/Replication/Editor/Model/Object/EditorObjectNameModel.h
@@ -12,7 +12,7 @@ namespace UE::ConcertClientSharedSlate
 	public:
 
 		//~ Begin IObjectNameModel Interface
-		virtual FText GetObjectDisplayName(const FSoftObjectPath& ObjectPath) const override;
+		virtual FText GetObjectDisplayName(const TSoftObjectPtr<>& ObjectPath) const override;
 		//~ End IObjectNameModel Interface
 	};
 }
diff --git a/Engine/Plugins/Developer/Concert/ConcertUI/ConcertClientSharedSlate/Source/ConcertClientSharedSlate/Private/Replication/Editor/View/PropertyTree/PropertyFilter_ByPropertyType.cpp b/Engine/Plugins/Developer/Concert/ConcertUI/ConcertClientSharedSlate/Source/ConcertClientSharedSlate/Private/Replication/Editor/View/PropertyTree/PropertyFilter_ByPropertyType.cpp
--- a/Engine/Plugins/Developer/Concert/ConcertUI/ConcertClientSharedSlate/Source/ConcertClientSharedSlate/Private/Replication/Editor/View/PropertyTree/PropertyFilter_ByPropertyType.cpp
+++ b/Engine/Plugins/Developer/Concert/ConcertUI/ConcertClientSharedSlate/Source/ConcertClientSharedSlate/Private/Replication/Editor/View/PropertyTree/PropertyFilter_ByPropertyType.cpp
@@ -11,7 +11,7 @@ namespace UE::ConcertClientSharedSlate
 {
 	bool FPropertyFilter_ByPropertyType::MatchesFilteredForProperty(const ConcertSharedSlate::FPropertyData& InItem) const
 	{
-		UClass* Class = InItem.GetOwningClass().TryLoadClass<UObject>();
+		UClass* Class = InItem.GetOwningClassPtr().LoadSynchronous();
 		if (!Class)
 		{
 			return false;
diff --git a/Engine/Plugins/Developer/Concert/ConcertUI/ConcertSharedSlate/Source/ConcertSharedSlate/Private/Replication/Editor/Utils/DisplayUtils.cpp b/Engine/Plugins/Developer/Concert/ConcertUI/ConcertSharedSlate/Source/ConcertSharedSlate/Private/Replication/Editor/Utils/DisplayUtils.cpp
--- a/Engine/Plugins/Developer/Concert/ConcertUI/ConcertSharedSlate/Source/ConcertSharedSlate/Private/Replication/Editor/Utils/DisplayUtils.cpp
+++ b/Engine/Plugins/Developer/Concert/ConcertUI/ConcertSharedSlate/Source/ConcertSharedSlate/Private/Replication/Editor/Utils/DisplayUtils.cpp
@@ -3,11 +3,13 @@
 #include "Replication/Editor/Utils/DisplayUtils.h"
 
 #include "ConcertLogGlobal.h"
+#include "Replication/PropertyResolutionCache.h"
 #include "Replication/Utils/ObjectUtils.h"
 #include "Replication/Editor/Model/IReplicationStreamModel.h"
 #include "Replication/Editor/Model/Object/IObjectNameModel.h"
 
 #include "Styling/SlateIconFinder.h"
+#include "UObject/SoftObjectPtr.h"
 
 #if WITH_EDITOR
 #include "AssetRegistry/AssetData.h"
@@ -22,22 +24,20 @@
 
 namespace UE::ConcertSharedSlate::DisplayUtils
 {
-	FText GetObjectDisplayText(const FSoftObjectPath& Object, IObjectNameModel* Model)
+	FText GetObjectDisplayText(const TSoftObjectPtr<>& Object, IObjectNameModel* Model)
 	{
 		if (Model)
 		{
 			return Model->GetObjectDisplayName(Object);
 		}
 		
-#if WITH_EDITOR
 		// Important! The object may not be loaded, yet. So only resolve, do not TryLoad. This could be if the asset is using a level that was not opened.
-		if (const UObject* LoadedObject = Object.ResolveObject())
+		if (const UObject* LoadedObject = Object.Get())
 		{
 			return FText::FromString(GetObjectDisplayString(*LoadedObject));
 		}
-#endif
 
-		return ExtractObjectDisplayTextFromPath(Object);
+		return ExtractObjectDisplayTextFromPath(Object.GetUniqueID());
 	}
 	
 	FText ExtractObjectDisplayTextFromPath(const FSoftObjectPath& Object)
@@ -112,25 +112,27 @@ namespace UE::ConcertSharedSlate::DisplayUtils
 		return FSlateIconFinder::FindIconForClass(Object.GetClass());
 	}
 	
-	FText GetPropertyDisplayText(const FConcertPropertyChain& Property, UStruct* Class /*=nullptr*/)
+	FText GetPropertyDisplayText(ConcertSyncCore::PropertyChain::FPropertyResolutionCache& Cache, const FConcertPropertyChain& Property, UStruct* Class /*=nullptr*/)
 	{
 #if WITH_EDITORONLY_DATA
 		// Class will (likely) be valid on editor builds but null on server
-		const FProperty* ResolvedProperty = Class ? Property.ResolveProperty(*Class) : nullptr;
+		const FProperty* ResolvedProperty = Class ? Cache.ResolveAndCache(*Class, Property) : nullptr;
 		return ResolvedProperty
 			? ResolvedProperty->GetDisplayNameText()
 			:
 #else
 		return 
 #endif
-		FText::FromString(GetPropertyDisplayString(Property, Class));
+		FText::FromString(GetPropertyDisplayString(Cache, Property, Class));
     }
 	
-	FString GetPropertyDisplayString(const FConcertPropertyChain& Property, UStruct* Class /*=nullptr*/)
+	FString GetPropertyDisplayString(ConcertSyncCore::PropertyChain::FPropertyResolutionCache& Cache, const FConcertPropertyChain& Property, UStruct* Class /*=nullptr*/)
 	{
+		SCOPED_CONCERT_TRACE(GetPropertyDisplayString);
+		
 #if WITH_EDITORONLY_DATA
 		// Class will (likely) be valid on editor builds but null on server
-		const FProperty* ResolvedProperty = Class ? Property.ResolveProperty(*Class) : nullptr;
+		const FProperty* ResolvedProperty = Class ? Cache.ResolveAndCache(*Class, Property) : nullptr;
 		return ResolvedProperty
 			? ResolvedProperty->GetDisplayNameText().ToString()
 			:
diff --git a/Engine/Plugins/Developer/Concert/ConcertUI/ConcertSharedSlate/Source/ConcertSharedSlate/Private/Replication/Editor/Utils/DisplayUtils.h b/Engine/Plugins/Developer/Concert/ConcertUI/ConcertSharedSlate/Source/ConcertSharedSlate/Private/Replication/Editor/Utils/DisplayUtils.h
--- a/Engine/Plugins/Developer/Concert/ConcertUI/ConcertSharedSlate/Source/ConcertSharedSlate/Private/Replication/Editor/Utils/DisplayUtils.h
+++ b/Engine/Plugins/Developer/Concert/ConcertUI/ConcertSharedSlate/Source/ConcertSharedSlate/Private/Replication/Editor/Utils/DisplayUtils.h
@@ -4,6 +4,7 @@
 
 #include "Internationalization/Text.h"
 #include "Templates/SharedPointer.h"
+#include "UObject/SoftObjectPtr.h"
 
 class FString;
 class UObject;
@@ -19,12 +20,13 @@ namespace UE::ConcertSharedSlate
 	class IObjectNameModel;
 	class IReplicationStreamModel;
 }
+namespace UE::ConcertSyncCore::PropertyChain { class FPropertyResolutionCache; }
 
 namespace UE::ConcertSharedSlate::DisplayUtils
 {
 	/** Unified version that reads from IObjectNameModel or defaults to extracting the name from the path. */
-	FText GetObjectDisplayText(const FSoftObjectPath& Object, IObjectNameModel* Model = nullptr);
-	inline FText GetObjectDisplayText(const FSoftObjectPath& Object, const TSharedPtr<IObjectNameModel>& Model = nullptr) { return GetObjectDisplayText(Object, Model.Get()); }
+	FText GetObjectDisplayText(const TSoftObjectPtr<>& Object, IObjectNameModel* Model = nullptr);
+	inline FText GetObjectDisplayText(const TSoftObjectPtr<>& Object, const TSharedPtr<IObjectNameModel>& Model = nullptr) { return GetObjectDisplayText(Object, Model.Get()); }
 	
 	/** @return The text to use for displaying this object's name */
 	FText ExtractObjectDisplayTextFromPath(const FSoftObjectPath& Object);
@@ -41,7 +43,7 @@ namespace UE::ConcertSharedSlate::DisplayUtils
 	FSlateIcon GetObjectIcon(UObject& Object);
 
 	/** @return The text to use for displaying this property's name. Uses Class to determine class name if available. */
-	FText GetPropertyDisplayText(const FConcertPropertyChain& Property, UStruct* Class = nullptr);
+	FText GetPropertyDisplayText(ConcertSyncCore::PropertyChain::FPropertyResolutionCache& Cache, const FConcertPropertyChain& Property, UStruct* Class = nullptr);
 	/** @return More lightweight version of GetPropertyDisplayString which does not construct any FText. Uses Class to determine class name if available. */
-	FString GetPropertyDisplayString(const FConcertPropertyChain& Property, UStruct* Class = nullptr);
+	FString GetPropertyDisplayString(ConcertSyncCore::PropertyChain::FPropertyResolutionCache& Cache, const FConcertPropertyChain& Property, UStruct* Class = nullptr);
 }
diff --git a/Engine/Plugins/Developer/Concert/ConcertUI/ConcertSharedSlate/Source/ConcertSharedSlate/Private/Replication/Editor/View/ObjectViewer/SReplicatedPropertyView.cpp b/Engine/Plugins/Developer/Concert/ConcertUI/ConcertSharedSlate/Source/ConcertSharedSlate/Private/Replication/Editor/View/ObjectViewer/SReplicatedPropertyView.cpp
--- a/Engine/Plugins/Developer/Concert/ConcertUI/ConcertSharedSlate/Source/ConcertSharedSlate/Private/Replication/Editor/View/ObjectViewer/SReplicatedPropertyView.cpp
+++ b/Engine/Plugins/Developer/Concert/ConcertUI/ConcertSharedSlate/Source/ConcertSharedSlate/Private/Replication/Editor/View/ObjectViewer/SReplicatedPropertyView.cpp
@@ -4,6 +4,7 @@
 
 #include "Replication/Editor/Model/IReplicationStreamModel.h"
 #include "Replication/Editor/Model/Data/ReplicatedObjectData.h"
+#include "Trace/ConcertTrace.h"
 
 #include "Algo/AllOf.h"
 #include "Widgets/Layout/SWidgetSwitcher.h"
@@ -30,6 +31,7 @@ namespace UE::ConcertSharedSlate
 	
 	void SReplicatedPropertyView::RefreshPropertyData()
 	{
+		SCOPED_CONCERT_TRACE(RefreshPropertyData);
 		TSet<FSoftObjectPath> SelectedObjects = GetObjectsSelectedForPropertyEditing();
 		if (SelectedObjects.IsEmpty())
 		{
@@ -95,7 +97,7 @@ namespace UE::ConcertSharedSlate
 		FSoftClassPath SharedClass;
 		const bool bAllHaveSameClass = Algo::AllOf(Objects, [this, &SharedClass](const FSoftObjectPath& Object)
 		{
-			const FSoftClassPath ObjectClass = GetObjectClass(Object);
+			const FSoftClassPath ObjectClass = GetObjectClass(TSoftObjectPtr<>{ Object });
 			SharedClass = SharedClass.IsValid() ? SharedClass : ObjectClass;
 			return ObjectClass == SharedClass;
 		});
diff --git a/Engine/Plugins/Developer/Concert/ConcertUI/ConcertSharedSlate/Source/ConcertSharedSlate/Private/Replication/Editor/View/ObjectViewer/SReplicatedPropertyView.h b/Engine/Plugins/Developer/Concert/ConcertUI/ConcertSharedSlate/Source/ConcertSharedSlate/Private/Replication/Editor/View/ObjectViewer/SReplicatedPropertyView.h
--- a/Engine/Plugins/Developer/Concert/ConcertUI/ConcertSharedSlate/Source/ConcertSharedSlate/Private/Replication/Editor/View/ObjectViewer/SReplicatedPropertyView.h
+++ b/Engine/Plugins/Developer/Concert/ConcertUI/ConcertSharedSlate/Source/ConcertSharedSlate/Private/Replication/Editor/View/ObjectViewer/SReplicatedPropertyView.h
@@ -76,7 +76,7 @@ namespace UE::ConcertSharedSlate
 		/** Sets how to display this widget */
 		void SetPropertyContent(EReplicatedPropertyContent Content) const;
 
-		FSoftClassPath GetObjectClass(const FSoftObjectPath& ObjectPath) const { return GetObjectClassDelegate.Execute(ObjectPath); };
+		FSoftClassPath GetObjectClass(const TSoftObjectPtr<>& Object) const { return GetObjectClassDelegate.Execute(Object); }
 	};
 }
 
diff --git a/Engine/Plugins/Developer/Concert/ConcertUI/ConcertSharedSlate/Source/ConcertSharedSlate/Private/Replication/Editor/View/ObjectViewer/SReplicationStreamViewer.cpp b/Engine/Plugins/Developer/Concert/ConcertUI/ConcertSharedSlate/Source/ConcertSharedSlate/Private/Replication/Editor/View/ObjectViewer/SReplicationStreamViewer.cpp
--- a/Engine/Plugins/Developer/Concert/ConcertUI/ConcertSharedSlate/Source/ConcertSharedSlate/Private/Replication/Editor/View/ObjectViewer/SReplicationStreamViewer.cpp
+++ b/Engine/Plugins/Developer/Concert/ConcertUI/ConcertSharedSlate/Source/ConcertSharedSlate/Private/Replication/Editor/View/ObjectViewer/SReplicationStreamViewer.cpp
@@ -11,6 +11,7 @@
 #include "Replication/Editor/View/Property/SPropertyTreeView.h"
 #include "Replication/Utils/ObjectUtils.h"
 #include "SReplicatedPropertyView.h"
+#include "Trace/ConcertTrace.h"
 
 #include "Widgets/Layout/SBorder.h"
 #include "Widgets/Layout/SSplitter.h"
@@ -64,6 +65,8 @@ namespace UE::ConcertSharedSlate
 
 	void SReplicationStreamViewer::SelectObjects(TConstArrayView<FSoftObjectPath> Objects, bool bAtEndOfTick)
 	{
+		SCOPED_CONCERT_TRACE(SelectObjects);
+		
 		if (bHasRequestedObjectRefresh || bAtEndOfTick)
 		{
 			PendingToSelect = Objects;
@@ -85,6 +88,8 @@ namespace UE::ConcertSharedSlate
 
 	void SReplicationStreamViewer::ExpandObjects(TConstArrayView<FSoftObjectPath> Objects, bool bRecursive, bool bAtEndOfTick)
 	{
+		SCOPED_CONCERT_TRACE(ExpandObjects);
+		
 		if (Objects.IsEmpty())
 		{
 			return;
@@ -153,6 +158,8 @@ namespace UE::ConcertSharedSlate
 
 	void SReplicationStreamViewer::Tick(const FGeometry& AllottedGeometry, const double InCurrentTime, const float InDeltaTime)
 	{
+		SCOPED_CONCERT_TRACE(TickReplicationStreamViewer);
+		
 		if (bHasRequestedObjectRefresh)
 		{
 			bHasRequestedObjectRefresh = false;
@@ -295,6 +302,8 @@ namespace UE::ConcertSharedSlate
 	
 	void SReplicationStreamViewer::RefreshObjectData()
 	{
+		SCOPED_CONCERT_TRACE(RefreshObjectData);
+		
 		// Re-using existing instances is tricky: we cannot update the object path in an item because the list view will no detect this change;
 		// list view only looks at the shared ptr address. So the UI will not be refreshed. Since the number of items will be small, just reallocate... 
 		AllObjectRowData.Empty();
@@ -490,9 +499,9 @@ namespace UE::ConcertSharedSlate
 		return !bSkipSubobject;
 	}
 
-	FSoftClassPath SReplicationStreamViewer::GetObjectClass(const FSoftObjectPath& ObjectPath) const
+	FSoftClassPath SReplicationStreamViewer::GetObjectClass(const TSoftObjectPtr<>& Object) const
 	{
-		const FSoftClassPath ResolvedClass = PropertiesModel->GetObjectClass(ObjectPath);
+		const FSoftClassPath ResolvedClass = PropertiesModel->GetObjectClass(Object.GetUniqueID());
 		if (ResolvedClass.IsValid())
 		{
 			return ResolvedClass;
@@ -501,7 +510,7 @@ namespace UE::ConcertSharedSlate
 #if WITH_EDITOR
 		// In the editor, we display the entire hierarchy (see BuildObjectHierarchyIfNeeded) so some items may not be in PropertiesModel.
 		// Example: Add an actor with many components and assign nothing - all of those components will take this path.
-		const UObject* LoadedObject = ObjectPath.ResolveObject();
+		const UObject* LoadedObject = Object.Get();
 		return LoadedObject ? LoadedObject->GetClass() : FSoftClassPath{};
 #else
 		// For non-editor, we should probably consider getting the class information through a delegate.
diff --git a/Engine/Plugins/Developer/Concert/ConcertUI/ConcertSharedSlate/Source/ConcertSharedSlate/Private/Replication/Editor/View/ObjectViewer/SReplicationStreamViewer.h b/Engine/Plugins/Developer/Concert/ConcertUI/ConcertSharedSlate/Source/ConcertSharedSlate/Private/Replication/Editor/View/ObjectViewer/SReplicationStreamViewer.h
--- a/Engine/Plugins/Developer/Concert/ConcertUI/ConcertSharedSlate/Source/ConcertSharedSlate/Private/Replication/Editor/View/ObjectViewer/SReplicationStreamViewer.h
+++ b/Engine/Plugins/Developer/Concert/ConcertUI/ConcertSharedSlate/Source/ConcertSharedSlate/Private/Replication/Editor/View/ObjectViewer/SReplicationStreamViewer.h
@@ -188,6 +188,6 @@ namespace UE::ConcertSharedSlate
 		/** Whether the view options allow this type of relationship to be shown. */
 		bool ShouldDisplayObjectRelation(EChildRelationship Relationship) const;
 		
-		FSoftClassPath GetObjectClass(const FSoftObjectPath& ObjectPath) const;
+		FSoftClassPath GetObjectClass(const TSoftObjectPtr<>& Object) const;
 	};
 }
diff --git a/Engine/Plugins/Developer/Concert/ConcertUI/ConcertSharedSlate/Source/ConcertSharedSlate/Private/Replication/Editor/View/SelectionViewerColumns.cpp b/Engine/Plugins/Developer/Concert/ConcertUI/ConcertSharedSlate/Source/ConcertSharedSlate/Private/Replication/Editor/View/SelectionViewerColumns.cpp
--- a/Engine/Plugins/Developer/Concert/ConcertUI/ConcertSharedSlate/Source/ConcertSharedSlate/Private/Replication/Editor/View/SelectionViewerColumns.cpp
+++ b/Engine/Plugins/Developer/Concert/ConcertUI/ConcertSharedSlate/Source/ConcertSharedSlate/Private/Replication/Editor/View/SelectionViewerColumns.cpp
@@ -12,6 +12,7 @@
 #include "Replication/PropertyChainUtils.h"
 
 #include "Internationalization/Internationalization.h"
+#include "Replication/PropertyResolutionCache.h"
 #include "Textures/SlateIcon.h"
 #include "Widgets/Images/SImage.h"
 #include "Widgets/Input/SCheckBox.h"
@@ -19,6 +20,7 @@
 #include "Widgets/SBoxPanel.h"
 #include "Widgets/Text/STextBlock.h"
 #include "Styling/AppStyle.h"
+#include "Trace/ConcertTrace.h"
 
 #define LOCTEXT_NAMESPACE "ReplicationObjectColumns"
 
@@ -49,12 +51,12 @@ namespace UE::ConcertSharedSlate::ReplicationColumns::TopLevel
 			virtual TSharedRef<SWidget> GenerateColumnWidget(const FBuildArgs& InArgs) override
 			{
 				const FReplicatedObjectData& ObjectData = InArgs.RowItem.RowData;
-				const FSoftObjectPath& ObjectPath = ObjectData.GetObjectPath();
+				const TSoftObjectPtr<>& Object = ObjectData.GetObjectPtr();
 				
 				const FText Text = GetDisplayText(ObjectData);
-				const FSlateIcon ClassIcon = GetObjectClassDelegate.IsBound() ? DisplayUtils::GetObjectIcon(GetObjectClassDelegate.Execute(ObjectPath)) : FSlateIcon{};
+				const FSlateIcon ClassIcon = GetObjectClassDelegate.IsBound() ? DisplayUtils::GetObjectIcon(GetObjectClassDelegate.Execute(Object)) : FSlateIcon{};
 				return SNew(SHorizontalBox)
-					.ToolTipText(FText::FromString(ObjectPath.ToString()))
+					.ToolTipText(FText::FromString(Object.ToString()))
 					
 					+SHorizontalBox::Slot()
 					.AutoWidth()
@@ -79,7 +81,9 @@ namespace UE::ConcertSharedSlate::ReplicationColumns::TopLevel
 			virtual void PopulateSearchString(const FObjectTreeRowContext& InItem, TArray<FString>& InOutSearchStrings) const override
 			{
 				const FReplicatedObjectData& ObjectData = InItem.RowData;
-				InOutSearchStrings.Add(DisplayUtils::GetObjectDisplayText(ObjectData.GetObjectPath(), OptionalNameModel).ToString());
+				InOutSearchStrings.Add(
+					DisplayUtils::GetObjectDisplayText(ObjectData.GetObjectPtr(), OptionalNameModel).ToString()
+					);
 			}
 			
 			virtual bool CanBeSorted() const override { return true; } 
@@ -95,8 +99,7 @@ namespace UE::ConcertSharedSlate::ReplicationColumns::TopLevel
 			
 			FText GetDisplayText(const FReplicatedObjectData& ObjectData) const
 			{
-				const FSoftObjectPath& ObjectPath = ObjectData.GetObjectPath();
-				return DisplayUtils::GetObjectDisplayText(ObjectPath, OptionalNameModel);
+				return DisplayUtils::GetObjectDisplayText(ObjectData.GetObjectPtr(), OptionalNameModel);
 			}
 		};
 
@@ -129,7 +132,7 @@ namespace UE::ConcertSharedSlate::ReplicationColumns::TopLevel
 			
 			virtual TSharedRef<SWidget> GenerateColumnWidget(const FBuildArgs& InArgs) override
 			{
-				const FSoftClassPath Class = GetClass(InArgs.RowItem.RowData.GetObjectPath());
+				const FSoftClassPath Class = GetClass(InArgs.RowItem.RowData.GetObjectPtr());
 				return SNew(SBox)
 					.Padding(8, 0, 0, 0) // So the type name text is aligned with the header column text
 					[
@@ -142,15 +145,15 @@ namespace UE::ConcertSharedSlate::ReplicationColumns::TopLevel
 			
 			virtual void PopulateSearchString(const FObjectTreeRowContext& InItem, TArray<FString>& InOutSearchStrings) const override
 			{
-				const FSoftClassPath Class = GetClass(InItem.RowData.GetObjectPath());
+				const FSoftClassPath Class = GetClass(InItem.RowData.GetObjectPtr());
 				InOutSearchStrings.Add(DisplayUtils::GetObjectTypeText(Class).ToString());
 			}
 			
 			virtual bool CanBeSorted() const override { return true; } 
 			virtual bool IsLessThan(const FObjectTreeRowContext& Left, const FObjectTreeRowContext& Right) const override
 			{
-				const FSoftClassPath LeftClass = GetClass(Left.RowData.GetObjectPath());
-				const FSoftClassPath RightClass = GetClass(Right.RowData.GetObjectPath());
+				const FSoftClassPath LeftClass = GetClass(Left.RowData.GetObjectPtr());
+				const FSoftClassPath RightClass = GetClass(Right.RowData.GetObjectPtr());
 				return DisplayUtils::GetObjectTypeText(LeftClass).ToString()
 					< DisplayUtils::GetObjectTypeText(RightClass).ToString();
 			}
@@ -159,7 +162,7 @@ namespace UE::ConcertSharedSlate::ReplicationColumns::TopLevel
 			
 			const FGetObjectClass GetObjectClassDelegate;
 
-			FSoftClassPath GetClass(const FSoftObjectPath& Path) const { return GetObjectClassDelegate.Execute(Path); }
+			FSoftClassPath GetClass(const TSoftObjectPtr<>& Object) const { return GetObjectClassDelegate.Execute(Object); }
 		};
 
 		check(GetObjectClassDelegate.IsBound());
@@ -290,34 +293,40 @@ namespace UE::ConcertSharedSlate::ReplicationColumns::Property
 				return SNew(STextBlock)
 					.Font(FAppStyle::GetFontStyle("PropertyWindow.NormalFont"))
 					.HighlightText(TAttribute<FText>::CreateLambda([HighlightText = InArgs.HighlightText](){ return *HighlightText; }))
-					.Text(DisplayUtils::GetPropertyDisplayText(PropertyData.GetProperty(), ResolveOrLoadClass(PropertyData)));
+					.Text(DisplayUtils::GetPropertyDisplayText(PropertyCache, PropertyData.GetProperty(), ResolveOrLoadClass(PropertyData)));
 			}
 			
 			virtual void PopulateSearchString(const FPropertyTreeRowContext& InItem, TArray<FString>& InOutSearchStrings) const override
 			{
 				const FPropertyData& PropertyData = InItem.RowData;
-				FString DisplayString = DisplayUtils::GetPropertyDisplayString(InItem.RowData.GetProperty(), ResolveOrLoadClass(PropertyData));
+				FString DisplayString = DisplayUtils::GetPropertyDisplayString(PropertyCache, InItem.RowData.GetProperty(), ResolveOrLoadClass(PropertyData));
 				InOutSearchStrings.Emplace(MoveTemp(DisplayString));
 			}
 			
 			virtual bool CanBeSorted() const override { return true; } 
 			virtual bool IsLessThan(const FPropertyTreeRowContext& Left, const FPropertyTreeRowContext& Right) const override
 			{
-				return DisplayUtils::GetPropertyDisplayString(Left.RowData.GetProperty(), ResolveOrLoadClass(Left.RowData))
-					< DisplayUtils::GetPropertyDisplayString(Right.RowData.GetProperty(), ResolveOrLoadClass(Right.RowData));
+				SCOPED_CONCERT_TRACE(IsLess_ReplicationLabel)
+				return DisplayUtils::GetPropertyDisplayString(PropertyCache, Left.RowData.GetProperty(), ResolveOrLoadClass(Left.RowData))
+					< DisplayUtils::GetPropertyDisplayString(PropertyCache, Right.RowData.GetProperty(), ResolveOrLoadClass(Right.RowData));
 			}
 
 		private:
 
+			/**
+			 * Maps FConcertPropertyChain to FProperty.
+			 *
+			 * This improves GetPropertyDisplayString performance.
+			 * This reduced time spent by ~75% when doing a full tree refresh.
+			 *
+			 * Mutable because the cache may be mutated but it does not inheritently change the state of this object.
+			 */
+			mutable ConcertSyncCore::PropertyChain::FPropertyResolutionCache PropertyCache;
+
 			static UClass* ResolveOrLoadClass(const FPropertyData& PropertyData)
 			{
-#if WITH_EDITOR
-				// On editor there may be Blueprints that need loading...
-				return PropertyData.GetOwningClass().TryLoadClass<UObject>();
-#else
-				// ... but on the server everything should be native C++ and hence already loaded
-				return PropertyData.GetOwningClass().ResolveClass();
-#endif
+				SCOPED_CONCERT_TRACE(ResolveOrLoadClass);
+				return PropertyData.GetOwningClassPtr().LoadSynchronous();
 			}
 		};
 		
@@ -366,6 +375,7 @@ namespace UE::ConcertSharedSlate::ReplicationColumns::Property
 			virtual bool CanBeSorted() const override { return true; } 
 			virtual bool IsLessThan(const FPropertyTreeRowContext& Left, const FPropertyTreeRowContext& Right) const override
 			{
+				SCOPED_CONCERT_TRACE(IsLess_ReplicationType)
 				return GetDisplayText(Left.RowData).ToString() < GetDisplayText(Right.RowData).ToString();
 			}
 
@@ -373,7 +383,7 @@ namespace UE::ConcertSharedSlate::ReplicationColumns::Property
 			
 			static FText GetDisplayText(const FPropertyData& Args)
 			{
-				UClass* Class = Args.GetOwningClass().TryLoadClass<UObject>();
+				UClass* Class = Args.GetOwningClassPtr().LoadSynchronous();
 				const FProperty* Property = Class ? ConcertSyncCore::PropertyChain::ResolveProperty(*Class, Args.GetProperty()) : nullptr;
 				return Property ? FText::FromString(Property->GetCPPType()) : LOCTEXT("Unknown", "Unknown");	
 			}
diff --git a/Engine/Plugins/Developer/Concert/ConcertUI/ConcertSharedSlate/Source/ConcertSharedSlate/Private/Replication/Editor/View/Tree/SReplicationTreeView.h b/Engine/Plugins/Developer/Concert/ConcertUI/ConcertSharedSlate/Source/ConcertSharedSlate/Private/Replication/Editor/View/Tree/SReplicationTreeView.h
--- a/Engine/Plugins/Developer/Concert/ConcertUI/ConcertSharedSlate/Source/ConcertSharedSlate/Private/Replication/Editor/View/Tree/SReplicationTreeView.h
+++ b/Engine/Plugins/Developer/Concert/ConcertUI/ConcertSharedSlate/Source/ConcertSharedSlate/Private/Replication/Editor/View/Tree/SReplicationTreeView.h
@@ -17,6 +17,7 @@
 #include "Widgets/Views/STreeView.h"
 #include "Styling/AppStyle.h"
 #include "Styling/CoreStyle.h"
+#include "Trace/ConcertTrace.h"
 
 #define LOCTEXT_NAMESPACE "SReplicationListView"
 
@@ -408,6 +409,8 @@ namespace UE::ConcertSharedSlate
 	template <typename TItemType>
 	void SReplicationTreeView<TItemType>::Tick(const FGeometry& AllottedGeometry, const double InCurrentTime, const float InDeltaTime)
 	{
+		SCOPED_CONCERT_TRACE(TickReplicationTree);
+		
 		if (bFilterChanged)
 		{
 			ReapplyFilters();
@@ -630,6 +633,8 @@ namespace UE::ConcertSharedSlate
 	template <typename TItemType>
 	void SReplicationTreeView<TItemType>::ReapplyFilters()
 	{
+		SCOPED_CONCERT_TRACE(ReapplyFilters);
+		
 		// Try preserving the selected activity.
 		TArray<TSharedPtr<TItemType>> SelectedItems = TreeView->GetSelectedItems();
 
@@ -753,6 +758,8 @@ namespace UE::ConcertSharedSlate
 	template <typename TItemType>
 	void SReplicationTreeView<TItemType>::Resort()
 	{
+		SCOPED_CONCERT_TRACE(Resort);
+		
 		Sort(FilteredRootItems);
 		// GetRowChildren will be called again, which will do the resort the children.
 		TreeView->RequestListRefresh();
@@ -761,6 +768,8 @@ namespace UE::ConcertSharedSlate
 	template <typename TItemType>
 	void SReplicationTreeView<TItemType>::Sort(TArray<TSharedPtr<TItemType>>& Items)
 	{
+		SCOPED_CONCERT_TRACE(SortReplicationTree);
+		
 		const auto IsLessThan = [this](const TSharedPtr<TItemType>& Left, const TSharedPtr<TItemType>& Right, const FName& ColumnName, EColumnSortMode::Type SortMode)
 		{
 			const TSharedPtr<IReplicationTreeColumn<TItemType>> Column = FindColumnByName(ColumnName);
diff --git a/Engine/Plugins/Developer/Concert/ConcertUI/ConcertSharedSlate/Source/ConcertSharedSlate/Public/Replication/Editor/Model/Data/PropertyData.h b/Engine/Plugins/Developer/Concert/ConcertUI/ConcertSharedSlate/Source/ConcertSharedSlate/Public/Replication/Editor/Model/Data/PropertyData.h
--- a/Engine/Plugins/Developer/Concert/ConcertUI/ConcertSharedSlate/Source/ConcertSharedSlate/Public/Replication/Editor/Model/Data/PropertyData.h
+++ b/Engine/Plugins/Developer/Concert/ConcertUI/ConcertSharedSlate/Source/ConcertSharedSlate/Public/Replication/Editor/Model/Data/PropertyData.h
@@ -12,19 +12,25 @@ namespace UE::ConcertSharedSlate
 	public:
 		
 		FPropertyData(FSoftClassPath OwningClass, FConcertPropertyChain Object)
-			: OwningClass(MoveTemp(OwningClass))
+			: OwningClassPtr(MoveTemp(OwningClass))
 			, Property(MoveTemp(Object))
 		{}
 		
 		const FConcertPropertyChain& GetProperty() const { return Property; }
-		const FSoftClassPath& GetOwningClass() const { return OwningClass; }
+		const TSoftClassPtr<>& GetOwningClassPtr() const { return OwningClassPtr; }
 
 	private:
 
 		/** The class with which the FProperty can be determined. */
-		FSoftClassPath OwningClass;
+		TSoftClassPtr<> OwningClassPtr;
 		
-		/** The property to be replicated */
+		/**
+		 * The property to be replicated.
+		 *
+		 * On the servers, this will usually not resolve to anything.
+		 * This was promoted to be TSoftObjectPtr so that any UI that resolves this object path automatically caches it.
+		 * In certain operations this can improve performance: e.g. when fully rebuilding the property tree, this saved 35% performance.
+		 */
 		FConcertPropertyChain Property;
 	};
 }
\ No newline at end of file
diff --git a/Engine/Plugins/Developer/Concert/ConcertUI/ConcertSharedSlate/Source/ConcertSharedSlate/Public/Replication/Editor/Model/Data/ReplicatedObjectData.h b/Engine/Plugins/Developer/Concert/ConcertUI/ConcertSharedSlate/Source/ConcertSharedSlate/Public/Replication/Editor/Model/Data/ReplicatedObjectData.h
--- a/Engine/Plugins/Developer/Concert/ConcertUI/ConcertSharedSlate/Source/ConcertSharedSlate/Public/Replication/Editor/Model/Data/ReplicatedObjectData.h
+++ b/Engine/Plugins/Developer/Concert/ConcertUI/ConcertSharedSlate/Source/ConcertSharedSlate/Public/Replication/Editor/Model/Data/ReplicatedObjectData.h
@@ -12,14 +12,21 @@ namespace UE::ConcertSharedSlate
 	public:
 
 		FReplicatedObjectData(FSoftObjectPath ObjectPath)
-			: ObjectPath(MoveTemp(ObjectPath))
+			: ObjectPtr(MoveTemp(ObjectPath))
 		{}
 		
-		const FSoftObjectPath& GetObjectPath() const { return ObjectPath; }
+		const FSoftObjectPath& GetObjectPath() const { return ObjectPtr.GetUniqueID(); }
+		const TSoftObjectPtr<>& GetObjectPtr() const { return ObjectPtr; }
 
 	private:
 
-		/** The replicated object */
-		FSoftObjectPath ObjectPath;
+		/**
+		 * The replicated object.
+		 *
+		 * On the servers, this will usually not resolve to anything.
+		 * This was promoted to be TSoftObjectPtr so that any UI that resolves this object path automatically caches it.
+		 * In certain operations this can improve performance: e.g. when fully rebuilding the property tree, this saved 35% performance.
+		 */
+		TSoftObjectPtr<> ObjectPtr;
 	};
 }
\ No newline at end of file
diff --git a/Engine/Plugins/Developer/Concert/ConcertUI/ConcertSharedSlate/Source/ConcertSharedSlate/Public/Replication/Editor/Model/Object/IObjectNameModel.h b/Engine/Plugins/Developer/Concert/ConcertUI/ConcertSharedSlate/Source/ConcertSharedSlate/Public/Replication/Editor/Model/Object/IObjectNameModel.h
--- a/Engine/Plugins/Developer/Concert/ConcertUI/ConcertSharedSlate/Source/ConcertSharedSlate/Public/Replication/Editor/Model/Object/IObjectNameModel.h
+++ b/Engine/Plugins/Developer/Concert/ConcertUI/ConcertSharedSlate/Source/ConcertSharedSlate/Public/Replication/Editor/Model/Object/IObjectNameModel.h
@@ -4,6 +4,9 @@
 
 #include "Delegates/Delegate.h"
 #include "Templates/SharedPointer.h"
+#include "UObject/SoftObjectPtr.h"
+
+class FText;
 
 namespace UE::ConcertSharedSlate
 {
@@ -18,7 +21,7 @@ namespace UE::ConcertSharedSlate
 	public:
 
 		/** @return Gets the display name for an object */
-		virtual FText GetObjectDisplayName(const FSoftObjectPath& ObjectPath) const = 0;
+		virtual FText GetObjectDisplayName(const TSoftObjectPtr<>& ObjectPath) const = 0;
 
 		// In the future, we could add an event here for when an object's display name changes.
 
diff --git a/Engine/Plugins/Developer/Concert/ConcertUI/ConcertSharedSlate/Source/ConcertSharedSlate/Public/Replication/Utils/ReplicationWidgetDelegates.h b/Engine/Plugins/Developer/Concert/ConcertUI/ConcertSharedSlate/Source/ConcertSharedSlate/Public/Replication/Utils/ReplicationWidgetDelegates.h
--- a/Engine/Plugins/Developer/Concert/ConcertUI/ConcertSharedSlate/Source/ConcertSharedSlate/Public/Replication/Utils/ReplicationWidgetDelegates.h
+++ b/Engine/Plugins/Developer/Concert/ConcertUI/ConcertSharedSlate/Source/ConcertSharedSlate/Public/Replication/Utils/ReplicationWidgetDelegates.h
@@ -6,6 +6,7 @@
 #include "Delegates/Delegate.h"
 #include "Delegates/DelegateCombinations.h"
 #include "Templates/Function.h"
+#include "UObject/SoftObjectPtr.h"
 
 class FMenuBuilder;
 enum class EBreakBehavior : uint8;
@@ -26,7 +27,7 @@ namespace UE::ConcertSharedSlate
 	DECLARE_DELEGATE_TwoParams(FExtendObjectMenu, FMenuBuilder&, TConstArrayView<FSoftObjectPath> ContextObjects);
 
 	/** Delegate for getting an object's class. */
-	DECLARE_DELEGATE_RetVal_OneParam(FSoftClassPath, FGetObjectClass, const FSoftObjectPath&);
+	DECLARE_DELEGATE_RetVal_OneParam(FSoftClassPath, FGetObjectClass, const TSoftObjectPtr<>&);
 
 	/** Delegate for deciding whether an object should be displayed. */
 	DECLARE_DELEGATE_RetVal_OneParam(bool, FShouldDisplayObject, const FSoftObjectPath& ObjectPath);
