Actor rows in Multi User replication tab now now show bin when hovered:
- Actors have both the bin button, which remove the entire hierarchy
- Subobjects only have the reset button

diff --git a/Engine/Plugins/Developer/Concert/ConcertUI/ConcertSharedSlate/Source/ConcertSharedSlate/Private/Replication/Editor/View/ObjectEditor/SBaseReplicationStreamEditor.cpp b/Engine/Plugins/Developer/Concert/ConcertUI/ConcertSharedSlate/Source/ConcertSharedSlate/Private/Replication/Editor/View/ObjectEditor/SBaseReplicationStreamEditor.cpp
--- a/Engine/Plugins/Developer/Concert/ConcertUI/ConcertSharedSlate/Source/ConcertSharedSlate/Private/Replication/Editor/View/ObjectEditor/SBaseReplicationStreamEditor.cpp
+++ b/Engine/Plugins/Developer/Concert/ConcertUI/ConcertSharedSlate/Source/ConcertSharedSlate/Private/Replication/Editor/View/ObjectEditor/SBaseReplicationStreamEditor.cpp
@@ -3,16 +3,20 @@
 #include "SBaseReplicationStreamEditor.h"
 
 #include "Model/Item/SourceModelBuilders.h"
+#include "Replication/Editor/Model/Data/ReplicatedObjectData.h"
+#include "Replication/Editor/Utils/DisplayUtils.h"
 #include "Replication/Editor/Model/IEditableReplicationStreamModel.h"
+#include "Replication/Editor/Model/Object/IObjectHierarchyModel.h"
 #include "Replication/Editor/Model/Property/IPropertySelectionSourceModel.h"
-#include "Replication/Editor/Model/Data/ReplicatedObjectData.h"
 #include "Replication/Editor/Model/ObjectSource/IObjectSelectionSourceModel.h"
-#include "Replication/Editor/Utils/DisplayUtils.h"
 #include "Replication/Editor/View/ObjectViewer/SReplicationStreamViewer.h"
 
 #include "Algo/AnyOf.h"
-#include "Replication/Editor/Model/Object/IObjectHierarchyModel.h"
+#include "Replication/Utils/ObjectUtils.h"
+#include "Styling/AppStyle.h"
 #include "UObject/Class.h"
+#include "Widgets/Images/SImage.h"
+#include "Widgets/Input/SButton.h"
 #include "Widgets/SBoxPanel.h"
 
 #define LOCTEXT_NAMESPACE "SBaseReplicationStreamEditor"
@@ -63,6 +67,7 @@ namespace UE::ConcertSharedSlate
 						InArgs._LeftOfObjectSearchBar.Widget
 					]
 				]
+				.GetHoveredRowContent(this, &SBaseReplicationStreamEditor::MakeHoveredRowContent)
 				.NoOutlinerObjects(LOCTEXT("NoObjects", "Add objects to replicate"))
 		];
 	}
@@ -93,6 +98,33 @@ namespace UE::ConcertSharedSlate
 		return ReplicationViewer->GetObjectsBeingPropertyEdited();
 	}
 
+	FHoverRowContent SBaseReplicationStreamEditor::MakeHoveredRowContent(const TSharedPtr<FReplicatedObjectData>& Data) const
+	{
+		const FSoftObjectPath Object = Data->GetObjectPath();
+		
+		if (ObjectUtils::IsActor(Object))
+		{
+			const TSharedRef<SWidget> Content =
+				SNew(SButton)
+				.ButtonStyle( FAppStyle::Get(), "SimpleButton")
+				.OnClicked_Lambda([this, WeakData = Data.ToWeakPtr()]()
+				{
+					if (const TSharedPtr<FReplicatedObjectData> DataPin = WeakData.Pin())
+					{
+						OnDeleteObjects({ DataPin });
+					}
+					return FReply::Handled();
+				})
+				[
+					SNew(SImage)
+					.Image(FAppStyle::GetBrush("Icons.Delete"))
+				];
+			return { Content, HAlign_Right };
+		}
+		
+		return { SNullWidget::NullWidget };
+	}
+
 	bool SBaseReplicationStreamEditor::IsEditingDisabled() const
 	{
 		return (IsEditingEnabledAttribute.IsBound() || IsEditingEnabledAttribute.IsSet())
diff --git a/Engine/Plugins/Developer/Concert/ConcertUI/ConcertSharedSlate/Source/ConcertSharedSlate/Private/Replication/Editor/View/ObjectEditor/SBaseReplicationStreamEditor.h b/Engine/Plugins/Developer/Concert/ConcertUI/ConcertSharedSlate/Source/ConcertSharedSlate/Private/Replication/Editor/View/ObjectEditor/SBaseReplicationStreamEditor.h
--- a/Engine/Plugins/Developer/Concert/ConcertUI/ConcertSharedSlate/Source/ConcertSharedSlate/Private/Replication/Editor/View/ObjectEditor/SBaseReplicationStreamEditor.h
+++ b/Engine/Plugins/Developer/Concert/ConcertUI/ConcertSharedSlate/Source/ConcertSharedSlate/Private/Replication/Editor/View/ObjectEditor/SBaseReplicationStreamEditor.h
@@ -29,6 +29,8 @@ namespace UE::ConcertSharedSlate
 	class SReplicationStreamViewer;
 	
 	enum class EReplicatedObjectChangeReason : uint8;
+
+	struct FHoverRowContent;
 	struct FSelectableObjectInfo;
 	
 	/**
@@ -128,6 +130,9 @@ namespace UE::ConcertSharedSlate
 		 */
 		bool bIsAddingFromSelection = false;
 
+		/** Generates a widget when an object row is hovered. Adds a bin and reset button. */
+		FHoverRowContent MakeHoveredRowContent(const TSharedPtr<FReplicatedObjectData>& Data) const;
+
 		bool IsEditingDisabled() const;
 		FText GetEditingDisabledText() const;
 		
@@ -145,6 +150,6 @@ namespace UE::ConcertSharedSlate
 		void AddObjectSourceContextMenuOptions(FMenuBuilder& MenuBuilder);
 		
 		// Utils for building item source widgets
-		ConcertSharedSlate::FSourceModelBuilders<FSelectableObjectInfo>::FItemPickerArgs MakeObjectSourceBuilderArgs();
+		FSourceModelBuilders<FSelectableObjectInfo>::FItemPickerArgs MakeObjectSourceBuilderArgs();
 	};
 }
diff --git a/Engine/Plugins/Developer/Concert/ConcertUI/ConcertSharedSlate/Source/ConcertSharedSlate/Private/Replication/Editor/View/ObjectViewer/SReplicationStreamViewer.cpp b/Engine/Plugins/Developer/Concert/ConcertUI/ConcertSharedSlate/Source/ConcertSharedSlate/Private/Replication/Editor/View/ObjectViewer/SReplicationStreamViewer.cpp
--- a/Engine/Plugins/Developer/Concert/ConcertUI/ConcertSharedSlate/Source/ConcertSharedSlate/Private/Replication/Editor/View/ObjectViewer/SReplicationStreamViewer.cpp
+++ b/Engine/Plugins/Developer/Concert/ConcertUI/ConcertSharedSlate/Source/ConcertSharedSlate/Private/Replication/Editor/View/ObjectViewer/SReplicationStreamViewer.cpp
@@ -266,6 +266,7 @@ namespace UE::ConcertSharedSlate
 				]
 			]
 			.NoItemsContent() [ SNew(STextBlock).Text(NoObjectsAttribute) ]
+			.GetHoveredRowContent(InArgs._GetHoveredRowContent)
 			.RowStyle(FAppStyle::Get(), "TableView.AlternatingRow");
 	}
 
diff --git a/Engine/Plugins/Developer/Concert/ConcertUI/ConcertSharedSlate/Source/ConcertSharedSlate/Private/Replication/Editor/View/ObjectViewer/SReplicationStreamViewer.h b/Engine/Plugins/Developer/Concert/ConcertUI/ConcertSharedSlate/Source/ConcertSharedSlate/Private/Replication/Editor/View/ObjectViewer/SReplicationStreamViewer.h
--- a/Engine/Plugins/Developer/Concert/ConcertUI/ConcertSharedSlate/Source/ConcertSharedSlate/Private/Replication/Editor/View/ObjectViewer/SReplicationStreamViewer.h
+++ b/Engine/Plugins/Developer/Concert/ConcertUI/ConcertSharedSlate/Source/ConcertSharedSlate/Private/Replication/Editor/View/ObjectViewer/SReplicationStreamViewer.h
@@ -73,9 +73,13 @@ namespace UE::ConcertSharedSlate
 			/** Called to generate the context menu for objects. */
 			SLATE_EVENT(FOnContextMenuOpening, OnObjectsContextMenuOpening)
 			
+		
 			/** Optional. Whether a given object should be displayed. If this returns false on an object, none of its children will be shown either. */
 			SLATE_EVENT(FShouldDisplayObject, ShouldDisplayObject)
 		
+			/** Optional. Gets the content to overlay on hovered rows; it covers the entire row. */
+			SLATE_EVENT(SReplicationTreeView<FReplicatedObjectData>::FGetHoveredRowContent, GetHoveredRowContent)
+			
 			/** Optional widget to add to the left of the object list search bar. */
 			SLATE_NAMED_SLOT(FArguments, LeftOfObjectSearchBar)
 			/** Optional widget to add to the right of the object list search bar. */
diff --git a/Engine/Plugins/Developer/Concert/ConcertUI/ConcertSharedSlate/Source/ConcertSharedSlate/Private/Replication/Editor/View/Tree/SReplicationColumnRow.h b/Engine/Plugins/Developer/Concert/ConcertUI/ConcertSharedSlate/Source/ConcertSharedSlate/Private/Replication/Editor/View/Tree/SReplicationColumnRow.h
--- a/Engine/Plugins/Developer/Concert/ConcertUI/ConcertSharedSlate/Source/ConcertSharedSlate/Private/Replication/Editor/View/Tree/SReplicationColumnRow.h
+++ b/Engine/Plugins/Developer/Concert/ConcertUI/ConcertSharedSlate/Source/ConcertSharedSlate/Private/Replication/Editor/View/Tree/SReplicationColumnRow.h
@@ -8,6 +8,7 @@
 #include "Delegates/DelegateCombinations.h"
 #include "Widgets/Layout/SBox.h"
 #include "Widgets/SNullWidget.h"
+#include "Widgets/SOverlay.h"
 #include "Widgets/Views/SExpanderArrow.h"
 #include "Widgets/Views/SListView.h"
 #include "Widgets/Views/STableRow.h"
@@ -22,10 +23,12 @@ namespace UE::ConcertSharedSlate
 	template<typename TListItemType>
 	class SReplicationColumnRow : public SMultiColumnTableRow<TSharedPtr<TListItemType>>
 	{
+		using Super = SMultiColumnTableRow<TSharedPtr<TListItemType>>;
 	public:
 
 		using FGetColumn = typename TReplicationTreeData<TListItemType>::FGetColumn;
 		using FOverrideColumnWidget = typename TReplicationTreeData<TListItemType>::FOverrideColumnWidget;
+		using FGetHoveredRowContent = typename TReplicationTreeData<TListItemType>::FGetHoveredRowContent;
 		
 		SLATE_BEGIN_ARGS(SReplicationColumnRow)
 			: _Style(&FCoreStyle::Get().GetWidgetStyle<FTableRowStyle>("TableView.Row"))
@@ -41,6 +44,9 @@ namespace UE::ConcertSharedSlate
 			 * This is useful, e.g. if you want to generate a separator widget between items.
 			 */
 			SLATE_EVENT(FOverrideColumnWidget, OverrideColumnWidget)
+		
+			/** Optional. Gets the content to overlay on hovered rows; it covers the entire row. */
+			SLATE_EVENT(FGetHoveredRowContent, GetHoveredRowContent)
 
 			/** The data to pass to TReplicationColumn::BuildColumnWidget. */
 			SLATE_ARGUMENT(TSharedPtr<TListItemType>, RowData)
@@ -50,7 +56,6 @@ namespace UE::ConcertSharedSlate
 
 			/** Style to use for rows, e.g. for making them alternate in grey */
 			SLATE_STYLE_ARGUMENT(FTableRowStyle, Style)
-		
 		SLATE_END_ARGS()
 
 		void Construct(
@@ -59,6 +64,7 @@ namespace UE::ConcertSharedSlate
 		{
 			ColumnGetterDelegate = InArgs._ColumnGetter;
 			OverrideColumnWidgetDelegate = InArgs._OverrideColumnWidget;
+			GetHoveredRowContent = InArgs._GetHoveredRowContent;
 			HighlightText = InArgs._HighlightText;
 			RowData = InArgs._RowData;
 			ExpandableColumnLabel = InArgs._ExpandableColumnLabel;
@@ -70,6 +76,43 @@ namespace UE::ConcertSharedSlate
 				InOwner
 				);
 		}
+		
+		virtual void ConstructChildren(ETableViewMode::Type InOwnerTableMode, const TAttribute<FMargin>& InPadding, const TSharedRef<SWidget>& InContent) override
+		{
+			Super::ConstructChildren(InOwnerTableMode, InPadding, InContent);
+
+			if (GetHoveredRowContent.IsBound())
+			{
+				const auto[HoverContent, HAlign] = GetHoveredRowContent.Execute(RowData);
+				this->ChildSlot
+				[
+					SNew(SOverlay)
+
+					// Generated columns
+					+SOverlay::Slot()
+					[
+						this->ChildSlot.GetWidget()
+					]
+
+					// Hovered overlay
+					+SOverlay::Slot()
+					[
+						SNew(SHorizontalBox)
+						
+						+SHorizontalBox::Slot() 
+						.Padding(InPadding)
+						.HAlign(HAlign)
+						[
+							SNew(SBox)
+							.Visibility_Lambda([this](){ return this->IsHovered() ? EVisibility::Visible : EVisibility::Collapsed; })
+							[
+								HoverContent
+							]
+						]
+					]
+				];
+			}
+		}
 
 		/** Generates the widget representing this row. */
 		virtual TSharedRef<SWidget> GenerateWidgetForColumn(const FName& ColumnName) override
@@ -122,13 +165,13 @@ namespace UE::ConcertSharedSlate
 						ColumnWidget
 					]
 				];
-	
 		}
 
 	private:
 		
 		FGetColumn ColumnGetterDelegate;
 		FOverrideColumnWidget OverrideColumnWidgetDelegate;
+		FGetHoveredRowContent GetHoveredRowContent;
 		TSharedPtr<FText> HighlightText;
 		TSharedPtr<TListItemType> RowData;
 		FName ExpandableColumnLabel;
diff --git a/Engine/Plugins/Developer/Concert/ConcertUI/ConcertSharedSlate/Source/ConcertSharedSlate/Private/Replication/Editor/View/Tree/SReplicationTreeView.h b/Engine/Plugins/Developer/Concert/ConcertUI/ConcertSharedSlate/Source/ConcertSharedSlate/Private/Replication/Editor/View/Tree/SReplicationTreeView.h
--- a/Engine/Plugins/Developer/Concert/ConcertUI/ConcertSharedSlate/Source/ConcertSharedSlate/Private/Replication/Editor/View/Tree/SReplicationTreeView.h
+++ b/Engine/Plugins/Developer/Concert/ConcertUI/ConcertSharedSlate/Source/ConcertSharedSlate/Private/Replication/Editor/View/Tree/SReplicationTreeView.h
@@ -46,12 +46,12 @@ namespace UE::ConcertSharedSlate
 	public:
 
 		using TOverrideColumnWidget = typename TReplicationTreeData<TItemType>::FOverrideColumnWidget;
+		using FGetHoveredRowContent = typename TReplicationTreeData<TItemType>::FGetHoveredRowContent;
 
 		DECLARE_DELEGATE_OneParam(FDeleteItems, const TArray<TSharedPtr<TItemType>>& SelectedItems);
 		DECLARE_DELEGATE_TwoParams(FGetItemChildren, TSharedPtr<TItemType> Item, TFunctionRef<void(TSharedPtr<TItemType>)> ProcessChild);
 		DECLARE_DELEGATE(FOnSelectionChanged);
 
-		/** @return Whether the item passes the filter, i.e. should be shown. */
 		DECLARE_DELEGATE_RetVal_OneParam(EItemFilterResult, FCustomFilter, const TItemType& Item);
 		DECLARE_DELEGATE_RetVal_OneParam(bool, FIsSearchableItem, const TSharedPtr<TItemType>& Item);
 
@@ -91,6 +91,9 @@ namespace UE::ConcertSharedSlate
 			SLATE_EVENT(TOverrideColumnWidget, OverrideColumnWidget)
 			/** Optional callback for determining whether this item can be searched. */
 			SLATE_EVENT(FIsSearchableItem, IsSearchableItem)
+		
+			/** Optional. Gets the content to overlay on hovered rows; it covers the entire row. */
+			SLATE_EVENT(FGetHoveredRowContent, GetHoveredRowContent)
 			
 			/** The columns this list should have */
 			SLATE_ARGUMENT(TArray<TReplicationColumnEntry<TItemType>>, Columns)
@@ -130,6 +133,7 @@ namespace UE::ConcertSharedSlate
 			CustomFilterDelegate = InArgs._FilterItem;
 			OverrideColumnWidget = InArgs._OverrideColumnWidget;
 			IsSearchableItemDelegate = InArgs._IsSearchableItem;
+			GetHoveredRowContentDelegate = InArgs._GetHoveredRowContent;
 			ExpandableColumnId = InArgs._ExpandableColumnLabel;
 			RowStyle = InArgs._RowStyle;
 			
@@ -351,6 +355,8 @@ namespace UE::ConcertSharedSlate
 		TOverrideColumnWidget OverrideColumnWidget;
 		/** Optional callback for determining whether this item can be filtered. If false, it will not be shown when searched. */
 		FIsSearchableItem IsSearchableItemDelegate;
+		/** Optional. The content to overlay on hovered rows; it covers the entire row. */
+		FGetHoveredRowContent GetHoveredRowContentDelegate;
 
 		/** Style to use for rows */
 		const FTableRowStyle* RowStyle = nullptr;
@@ -515,6 +521,7 @@ namespace UE::ConcertSharedSlate
 			{
 				ColumnGetter,
 				OverrideColumnWidget,
+				GetHoveredRowContentDelegate,
 				SearchText,
 				ExpandableColumnId,
 				RowStyle
diff --git a/Engine/Plugins/Developer/Concert/ConcertUI/ConcertSharedSlate/Source/ConcertSharedSlate/Private/Replication/Editor/View/Tree/TreeItemTraits.h b/Engine/Plugins/Developer/Concert/ConcertUI/ConcertSharedSlate/Source/ConcertSharedSlate/Private/Replication/Editor/View/Tree/TreeItemTraits.h
--- a/Engine/Plugins/Developer/Concert/ConcertUI/ConcertSharedSlate/Source/ConcertSharedSlate/Private/Replication/Editor/View/Tree/TreeItemTraits.h
+++ b/Engine/Plugins/Developer/Concert/ConcertUI/ConcertSharedSlate/Source/ConcertSharedSlate/Private/Replication/Editor/View/Tree/TreeItemTraits.h
@@ -17,6 +17,7 @@ namespace UE::ConcertSharedSlate
 			.HighlightText(AdditionalArgs.HighlightText)
 			.ColumnGetter(AdditionalArgs.GetColumnDelegate)
 			.OverrideColumnWidget(AdditionalArgs.OverrideColumnWidgetDelegate)
+			.GetHoveredRowContent(AdditionalArgs.GetHoveredRowContent)
 			.RowData(Item)
 			.ExpandableColumnLabel(AdditionalArgs.ExpandableColumnId)
 			.Style(AdditionalArgs.RowStyle);
diff --git a/Engine/Plugins/Developer/Concert/ConcertUI/ConcertSharedSlate/Source/ConcertSharedSlate/Private/Replication/Editor/View/Tree/TreeItemTraitsInput.h b/Engine/Plugins/Developer/Concert/ConcertUI/ConcertSharedSlate/Source/ConcertSharedSlate/Private/Replication/Editor/View/Tree/TreeItemTraitsInput.h
--- a/Engine/Plugins/Developer/Concert/ConcertUI/ConcertSharedSlate/Source/ConcertSharedSlate/Private/Replication/Editor/View/Tree/TreeItemTraitsInput.h
+++ b/Engine/Plugins/Developer/Concert/ConcertUI/ConcertSharedSlate/Source/ConcertSharedSlate/Private/Replication/Editor/View/Tree/TreeItemTraitsInput.h
@@ -4,6 +4,7 @@
 
 #include "Delegates/Delegate.h"
 #include "Templates/SharedPointer.h"
+#include "Types/SlateEnums.h"
 
 class FName;
 class SWidget;
@@ -15,6 +16,12 @@ namespace UE::ConcertSharedSlate
 	template<typename TTreeItemType>
 	class IReplicationTreeColumn;
 
+	struct FHoverRowContent
+	{
+		TSharedRef<SWidget> Widget;
+		EHorizontalAlignment Alignment = HAlign_Right;
+	};
+	
 	/** Holds type definitions that server as input to TReplicationTreeItemTraits (and thus should not be template specialized).  */
 	template<typename TItemType>
 	class TReplicationTreeData
@@ -23,6 +30,7 @@ namespace UE::ConcertSharedSlate
 		
 		DECLARE_DELEGATE_RetVal_OneParam(TSharedPtr<IReplicationTreeColumn<TItemType>>, FGetColumn, const FName& ColumnId);
 		DECLARE_DELEGATE_RetVal_TwoParams(TSharedPtr<SWidget>, FOverrideColumnWidget, const FName& ColumnName, const TItemType& RowData);
+		DECLARE_DELEGATE_RetVal_OneParam(FHoverRowContent, FGetHoveredRowContent, const TSharedPtr<TItemType>& Item);
 		
 		struct FGenerateRowArgs
 		{
@@ -30,6 +38,9 @@ namespace UE::ConcertSharedSlate
 			FGetColumn GetColumnDelegate;
 			/** Overrides a column's content widget. */
 			FOverrideColumnWidget OverrideColumnWidgetDelegate;
+			
+			/** Optional. Gets the content to overlay on hovered rows; it covers the entire row. */
+			FGetHoveredRowContent GetHoveredRowContent;
 
 			/** The text to highlight - equal to search text. */
 			TSharedPtr<FText> HighlightText;
