Refactor FPropertyData to contain the objects for which the properties are being displayed. IPropertyTreeView to accept

diff --git a/Engine/Plugins/Developer/Concert/ConcertScripting/ConcertReplicationScripting/Source/ConcertReplicationScriptingEditor/Private/Widgets/SConcertPropertyChainPicker.cpp b/Engine/Plugins/Developer/Concert/ConcertScripting/ConcertReplicationScripting/Source/ConcertReplicationScriptingEditor/Private/Widgets/SConcertPropertyChainPicker.cpp
--- a/Engine/Plugins/Developer/Concert/ConcertScripting/ConcertReplicationScripting/Source/ConcertReplicationScriptingEditor/Private/Widgets/SConcertPropertyChainPicker.cpp
+++ b/Engine/Plugins/Developer/Concert/ConcertScripting/ConcertReplicationScripting/Source/ConcertReplicationScriptingEditor/Private/Widgets/SConcertPropertyChainPicker.cpp
@@ -121,8 +121,8 @@ namespace UE::ConcertReplicationScriptingEditor
 
 	void SConcertPropertyChainPicker::RefreshPropertiesDisplayedInTree()
 	{
-		TSet<FConcertPropertyChain> Properties;
-		FSoftClassPath ClassPath = SelectedClass;
+		ConcertSharedSlate::FPropertyAssignmentEntry Entry { .Class = SelectedClass };
+		TSet<FConcertPropertyChain>& Properties = Entry.PropertiesToDisplay;
 		
 		if (SelectedClass)
 		{
@@ -134,7 +134,7 @@ namespace UE::ConcertReplicationScriptingEditor
 		}
 
 		// If no class is selected, the widget will display the NoItemsContent ("Select a class").
-		TreeView->RefreshPropertyData(Properties, ClassPath);
+		TreeView->RefreshPropertyData({ Entry });
 	}
 
 	void SConcertPropertyChainPicker::OnPropertySelected(const FConcertPropertyChain& ConcertPropertyChain, bool bIsSelected)
diff --git a/Engine/Plugins/Developer/Concert/ConcertUI/ConcertClientSharedSlate/Source/ConcertClientSharedSlate/Private/Replication/Editor/View/PropertyTree/SFilteredPropertyTreeView.h b/Engine/Plugins/Developer/Concert/ConcertUI/ConcertClientSharedSlate/Source/ConcertClientSharedSlate/Private/Replication/Editor/View/PropertyTree/SFilteredPropertyTreeView.h
--- a/Engine/Plugins/Developer/Concert/ConcertUI/ConcertClientSharedSlate/Source/ConcertClientSharedSlate/Private/Replication/Editor/View/PropertyTree/SFilteredPropertyTreeView.h
+++ b/Engine/Plugins/Developer/Concert/ConcertUI/ConcertClientSharedSlate/Source/ConcertClientSharedSlate/Private/Replication/Editor/View/PropertyTree/SFilteredPropertyTreeView.h
@@ -36,9 +36,9 @@ namespace UE::ConcertClientSharedSlate
 		void Construct(const FArguments& InArgs, FFilterablePropertyTreeViewParams Params);
 		
 		//~ Begin IPropertyTreeView Interface
-		virtual void RefreshPropertyData(const TSet<FConcertPropertyChain>& PropertiesToDisplay, const FSoftClassPath& Class, bool bCanReuseExistingRowItems) override
+		virtual void RefreshPropertyData(const TArray<ConcertSharedSlate::FPropertyAssignmentEntry>& Entries, bool bCanReuseExistingRowItems) override
 		{
-			ExtendedTreeView->RefreshPropertyData(PropertiesToDisplay, Class, bCanReuseExistingRowItems);
+			ExtendedTreeView->RefreshPropertyData(Entries, bCanReuseExistingRowItems);
 		}
 		virtual void RequestRefilter() const override { ExtendedTreeView->RequestRefilter(); }
 		virtual void RequestResortForColumn(const FName& ColumnId) override { ExtendedTreeView->RequestResortForColumn(ColumnId); }
diff --git a/Engine/Plugins/Developer/Concert/ConcertUI/ConcertSharedSlate/Source/ConcertSharedSlate/ConcertSharedSlate.Build.cs b/Engine/Plugins/Developer/Concert/ConcertUI/ConcertSharedSlate/Source/ConcertSharedSlate/ConcertSharedSlate.Build.cs
--- a/Engine/Plugins/Developer/Concert/ConcertUI/ConcertSharedSlate/Source/ConcertSharedSlate/ConcertSharedSlate.Build.cs
+++ b/Engine/Plugins/Developer/Concert/ConcertUI/ConcertSharedSlate/Source/ConcertSharedSlate/ConcertSharedSlate.Build.cs
@@ -12,10 +12,11 @@ namespace UnrealBuildTool.Rules
 				new string[]
 				{
 					"Concert",
+					"ConcertSyncCore",
 					"Core",
 					"CoreUObject",
 					"Slate",
-					"UndoHistory",
+					"UndoHistory", 
 				}
 			);
 
diff --git a/Engine/Plugins/Developer/Concert/ConcertUI/ConcertSharedSlate/Source/ConcertSharedSlate/Private/Replication/Editor/View/Property/SPerObjectPropertyAssignment.cpp b/Engine/Plugins/Developer/Concert/ConcertUI/ConcertSharedSlate/Source/ConcertSharedSlate/Private/Replication/Editor/View/Property/SPerObjectPropertyAssignment.cpp
--- a/Engine/Plugins/Developer/Concert/ConcertUI/ConcertSharedSlate/Source/ConcertSharedSlate/Private/Replication/Editor/View/Property/SPerObjectPropertyAssignment.cpp
+++ b/Engine/Plugins/Developer/Concert/ConcertUI/ConcertSharedSlate/Source/ConcertSharedSlate/Private/Replication/Editor/View/Property/SPerObjectPropertyAssignment.cpp
@@ -21,10 +21,14 @@ namespace UE::ConcertSharedSlate
 
 	void SPerObjectPropertyAssignment::RefreshData(const TSet<FSoftObjectPath>& Objects, const IReplicationStreamModel& Model)
 	{
-		TSet<FConcertPropertyChain> Properties;
+		TSet<TSoftObjectPtr<>> ContextObjects;
+		Algo::Transform(Objects, ContextObjects, [](const FSoftObjectPath& Path){ return TSoftObjectPtr{ Path }; });
+		
+		FPropertyAssignmentEntry AssignmentEntry { .ContextObjects = MoveTemp(ContextObjects) };
+		TSet<FConcertPropertyChain>& Properties = AssignmentEntry.PropertiesToDisplay;
+		FSoftClassPath& ClassPath = AssignmentEntry.Class;
+		
 		bool bHasClassPath = false;
-		FSoftClassPath ClassPath;
-
 		EnumerateProperties(Objects, Model, OptionalPropertySource.Get(),
 			[&Properties, &bHasClassPath, &ClassPath](const FSoftClassPath& Class, const FConcertPropertyChain& Chain)
 			{
@@ -43,7 +47,7 @@ namespace UE::ConcertSharedSlate
 				Properties.Add(Chain);
 				return EBreakBehavior::Continue;
 			});
-
+		
 		if (bHasClassPath)
 		{
 			// If the objects have changed, the classes may share properties.
@@ -54,12 +58,12 @@ namespace UE::ConcertSharedSlate
 			const bool bCanReusePropertyData = PreviousSelectedObjects.Num() == Objects.Num() && PreviousSelectedObjects.Includes(Objects);
 			PreviousSelectedObjects = Objects;
 			
-			TreeView->RefreshPropertyData(Properties, ClassPath, bCanReusePropertyData);
+			TreeView->RefreshPropertyData({ AssignmentEntry } , bCanReusePropertyData);
 		}
 		else
 		{
 			PreviousSelectedObjects.Reset();
-			TreeView->RefreshPropertyData(Properties, ClassPath, false);
+			TreeView->RefreshPropertyData({}, false);
 		}
 	}
 }
diff --git a/Engine/Plugins/Developer/Concert/ConcertUI/ConcertSharedSlate/Source/ConcertSharedSlate/Private/Replication/Editor/View/Property/SPropertyTreeView.cpp b/Engine/Plugins/Developer/Concert/ConcertUI/ConcertSharedSlate/Source/ConcertSharedSlate/Private/Replication/Editor/View/Property/SPropertyTreeView.cpp
--- a/Engine/Plugins/Developer/Concert/ConcertUI/ConcertSharedSlate/Source/ConcertSharedSlate/Private/Replication/Editor/View/Property/SPropertyTreeView.cpp
+++ b/Engine/Plugins/Developer/Concert/ConcertUI/ConcertSharedSlate/Source/ConcertSharedSlate/Private/Replication/Editor/View/Property/SPropertyTreeView.cpp
@@ -32,7 +32,7 @@ namespace UE::ConcertSharedSlate
 		];
 	}
 
-	void SPropertyTreeView::RefreshPropertyData(const TSet<FConcertPropertyChain>& PropertiesToDisplay, const FSoftClassPath& Class, const bool bCanReuseExistingRowItems)
+	void SPropertyTreeView::RefreshPropertyData(const TArray<FPropertyAssignmentEntry>& Entries, bool bCanReuseExistingRowItems)
 	{
 		if (!bCanReuseExistingRowItems)
 		{
@@ -43,12 +43,15 @@ namespace UE::ConcertSharedSlate
 		TMap<FConcertPropertyChain, TSharedPtr<FPropertyData>> NewChainToPropertyDataCache;
 		
 		PropertyRowData.Empty();
-		for (const FConcertPropertyChain& PropertyChain : PropertiesToDisplay)
+		for (const FPropertyAssignmentEntry& Entry : Entries)
 		{
-			const TSharedPtr<FPropertyData>* ExistingItem = ChainToPropertyDataCache.Find(PropertyChain);
-			const TSharedRef<FPropertyData> Item = ExistingItem ? ExistingItem->ToSharedRef() : AllocatePropertyData(Class, PropertyChain);
-			PropertyRowData.Emplace(Item);
-			NewChainToPropertyDataCache.Emplace(PropertyChain, Item);
+			for (const FConcertPropertyChain& PropertyChain : Entry.PropertiesToDisplay)
+			{
+				const TSharedPtr<FPropertyData>* ExistingItem = ChainToPropertyDataCache.Find(PropertyChain);
+				const TSharedRef<FPropertyData> Item = ExistingItem ? ExistingItem->ToSharedRef() : AllocatePropertyData(Entry.ContextObjects, Entry.Class, PropertyChain);
+				PropertyRowData.Emplace(Item);
+				NewChainToPropertyDataCache.Emplace(PropertyChain, Item);
+			}
 		}
 
 		// If an item was removed, then NewPathToPropertyDataCache does not contain it. 
@@ -73,9 +76,9 @@ namespace UE::ConcertSharedSlate
 		}
 	}
 
-	TSharedRef<FPropertyData> SPropertyTreeView::AllocatePropertyData(FSoftClassPath OwningClass, FConcertPropertyChain PropertyChain)
+	inline TSharedRef<FPropertyData> SPropertyTreeView::AllocatePropertyData(TSet<TSoftObjectPtr<>> ContextObjects, FSoftClassPath OwningClass, FConcertPropertyChain PropertyChain)
 	{
-		return MakeShared<FPropertyData>(MoveTemp(OwningClass), MoveTemp(PropertyChain));
+		return MakeShared<FPropertyData>(MoveTemp(ContextObjects), MoveTemp(OwningClass), MoveTemp(PropertyChain));
 	}
 
 	void SPropertyTreeView::BuildRootPropertyRowData()
diff --git a/Engine/Plugins/Developer/Concert/ConcertUI/ConcertSharedSlate/Source/ConcertSharedSlate/Private/Replication/Editor/View/Property/SPropertyTreeView.h b/Engine/Plugins/Developer/Concert/ConcertUI/ConcertSharedSlate/Source/ConcertSharedSlate/Private/Replication/Editor/View/Property/SPropertyTreeView.h
--- a/Engine/Plugins/Developer/Concert/ConcertUI/ConcertSharedSlate/Source/ConcertSharedSlate/Private/Replication/Editor/View/Property/SPropertyTreeView.h
+++ b/Engine/Plugins/Developer/Concert/ConcertUI/ConcertSharedSlate/Source/ConcertSharedSlate/Private/Replication/Editor/View/Property/SPropertyTreeView.h
@@ -59,7 +59,7 @@ namespace UE::ConcertSharedSlate
 		void Construct(const FArguments& InArgs);
 
 		//~ Begin IPropertyTreeView Interface
-		virtual void RefreshPropertyData(const TSet<FConcertPropertyChain>& PropertiesToDisplay, const FSoftClassPath& Class, bool bCanReuseExistingRowItems) override;
+		virtual void RefreshPropertyData(const TArray<FPropertyAssignmentEntry>& Entries, bool bCanReuseExistingRowItems) override;
 		virtual void RequestRefilter() const override { TreeView->RequestRefilter(); }
 		virtual void RequestResortForColumn(const FName& ColumnId) override { TreeView->RequestResortForColumn(ColumnId); }
 		virtual void RequestScrollIntoView(const FConcertPropertyChain& PropertyChain) override;
@@ -78,7 +78,7 @@ namespace UE::ConcertSharedSlate
 		/** Inverse map of PropertyRowData using FPropertyData::GetProperty as key. Contains all elements of PropertyRowData. */
 		TMap<FConcertPropertyChain, TSharedPtr<FPropertyData>> ChainToPropertyDataCache;
 		
-		TSharedRef<FPropertyData> AllocatePropertyData(FSoftClassPath OwningClass, FConcertPropertyChain PropertyChain);
+		TSharedRef<FPropertyData> AllocatePropertyData(TSet<TSoftObjectPtr<>> ContextObjects, FSoftClassPath OwningClass, FConcertPropertyChain PropertyChain);
 
 		/** Inits RootPropertyRowData from PropertyRowData. */
 		void BuildRootPropertyRowData();
diff --git a/Engine/Plugins/Developer/Concert/ConcertUI/ConcertSharedSlate/Source/ConcertSharedSlate/Private/Replication/Editor/View/Tree/SReplicationTreeView.h b/Engine/Plugins/Developer/Concert/ConcertUI/ConcertSharedSlate/Source/ConcertSharedSlate/Private/Replication/Editor/View/Tree/SReplicationTreeView.h
--- a/Engine/Plugins/Developer/Concert/ConcertUI/ConcertSharedSlate/Source/ConcertSharedSlate/Private/Replication/Editor/View/Tree/SReplicationTreeView.h
+++ b/Engine/Plugins/Developer/Concert/ConcertUI/ConcertSharedSlate/Source/ConcertSharedSlate/Private/Replication/Editor/View/Tree/SReplicationTreeView.h
@@ -37,7 +37,7 @@ namespace UE::ConcertSharedSlate
 	 * Shared code for the tree view for replicated actors and properties.
 	 * 
 	 * It is a tree view that is searchable with a search box and exposes slots to add more filter widgets, such as SBasicFilterBar.
-	 * The columns are modular and are abstracted by IReplicationTreeColumn. You can customize row content by by specializing TReplicationTreeItemTraits;
+	 * The columns are modular and are abstracted by IReplicationTreeColumn. You can customize row content by specializing TReplicationTreeItemTraits;
 	 * the default implementation creates a SReplicationColumnRow, which uses IReplicationTreeColumn to draw the columns. 
 	 */
 	template<typename TItemType>
diff --git a/Engine/Plugins/Developer/Concert/ConcertUI/ConcertSharedSlate/Source/ConcertSharedSlate/Public/Replication/Editor/Model/Data/PropertyData.h b/Engine/Plugins/Developer/Concert/ConcertUI/ConcertSharedSlate/Source/ConcertSharedSlate/Public/Replication/Editor/Model/Data/PropertyData.h
--- a/Engine/Plugins/Developer/Concert/ConcertUI/ConcertSharedSlate/Source/ConcertSharedSlate/Public/Replication/Editor/Model/Data/PropertyData.h
+++ b/Engine/Plugins/Developer/Concert/ConcertUI/ConcertSharedSlate/Source/ConcertSharedSlate/Public/Replication/Editor/Model/Data/PropertyData.h
@@ -3,6 +3,8 @@
 #pragma once
 
 #include "Replication/Data/ConcertPropertySelection.h"
+
+#include "Containers/Array.h"
 #include "UObject/SoftObjectPtr.h"
 
 namespace UE::ConcertSharedSlate
@@ -12,19 +14,29 @@ namespace UE::ConcertSharedSlate
 	{
 	public:
 		
-		FPropertyData(FSoftClassPath OwningClass, FConcertPropertyChain Object)
-			: OwningClassPtr(MoveTemp(OwningClass))
+		FPropertyData(TSet<TSoftObjectPtr<>> ContextObjects, FSoftClassPath OwningClass, FConcertPropertyChain Object)
+			: ContextObjects(MoveTemp(ContextObjects))
+			, OwningClassPtr(MoveTemp(OwningClass))
 			, Property(MoveTemp(Object))
 		{}
 		
+		const TSet<TSoftObjectPtr<>>& GetContextObjects() const { return ContextObjects; }
 		const FConcertPropertyChain& GetProperty() const { return Property; }
 		const TSoftClassPtr<>& GetOwningClassPtr() const { return OwningClassPtr; }
 
 	private:
 
+		/**
+		 * The objects for which the properties are being displayed.
+		 *
+		 * This usually has only 1 entry.
+		 * This has multiple elements in the case of multi-edit (i.e. when the user clicks multiple, compatible actors in the top-view).
+		 * For example, for multi-edit this could contain ActorA->StaticMeshComponent0 and ActorB->StaticMeshComponent0.
+		 */
+		TSet<TSoftObjectPtr<>> ContextObjects;
+		
 		/** The class with which the FProperty can be determined. */
 		TSoftClassPtr<> OwningClassPtr;
-		
 		/**
 		 * The property to be replicated.
 		 *
diff --git a/Engine/Plugins/Developer/Concert/ConcertUI/ConcertSharedSlate/Source/ConcertSharedSlate/Public/Replication/Editor/View/IPropertyTreeView.h b/Engine/Plugins/Developer/Concert/ConcertUI/ConcertSharedSlate/Source/ConcertSharedSlate/Public/Replication/Editor/View/IPropertyTreeView.h
--- a/Engine/Plugins/Developer/Concert/ConcertUI/ConcertSharedSlate/Source/ConcertSharedSlate/Public/Replication/Editor/View/IPropertyTreeView.h
+++ b/Engine/Plugins/Developer/Concert/ConcertUI/ConcertSharedSlate/Source/ConcertSharedSlate/Public/Replication/Editor/View/IPropertyTreeView.h
@@ -2,16 +2,34 @@
 
 #pragma once
 
+#include "Replication/Data/ConcertPropertySelection.h"
+
 #include "Containers/Set.h"
 #include "Templates/SharedPointer.h"
+#include "UObject/SoftObjectPath.h"
+#include "UObject/SoftObjectPtr.h"
 
 class SWidget;
 
-struct FConcertPropertyChain;
-struct FSoftClassPath;
-
 namespace UE::ConcertSharedSlate
 {
+	struct FPropertyAssignmentEntry
+	{
+		/**
+		 * The objects for which the properties are being displayed.
+		 *
+		 * This usually has only 1 entry.
+		 * This has multiple elements in the case of multi-edit (i.e. when the user clicks multiple, compatible actors in the top-view).
+		 * For example, for multi-edit this could contain ActorA->StaticMeshComponent0 and ActorB->StaticMeshComponent0.
+		 */
+		TSet<TSoftObjectPtr<>> ContextObjects;
+		
+		/** The properties to display */
+		TSet<FConcertPropertyChain> PropertiesToDisplay;
+		/** The class of the properties */
+		FSoftClassPath Class;
+	};
+	
 	/** Represents a tree view displaying properties from a single class. */
 	class CONCERTSHAREDSLATE_API IPropertyTreeView
 	{
@@ -26,7 +44,18 @@ namespace UE::ConcertSharedSlate
 		 *	Set this to false, if all rows should be regenerated (clears selection).
 		 *	In general, always set this to false if you've changed the object for which you're displaying the class.
 		 */
-		virtual void RefreshPropertyData(const TSet<FConcertPropertyChain>& PropertiesToDisplay, const FSoftClassPath& Class, bool bCanReuseExistingRowItems = true) = 0;
+		UE_DEPRECATED(5.5, "Use the version that takes FPropertyAssignmentEntry instead")
+		void RefreshPropertyData(const TSet<FConcertPropertyChain>& PropertiesToDisplay, const FSoftClassPath& Class, bool bCanReuseExistingRowItems = true) {}
+
+		/**
+		 * Rebuilds all property data from the property source.
+		 *
+		 * @param Entries Defines the property content to display
+		 * @param bCanReuseExistingRowItems True, will try to reuse rows for properties in the tree already (retains selected rows).
+		 *	Set this to false, if all rows should be regenerated (clears selection).
+		 *	In general, always set this to false if you've changed the object for which you're displaying the class.
+		 */
+		virtual void RefreshPropertyData(const TArray<FPropertyAssignmentEntry>& Entries, bool bCanReuseExistingRowItems = true) = 0;
 		
 		/** Reapply the filter function to all items at the end of the frame. Call e.g. when the filters have changed. */
 		virtual void RequestRefilter() const = 0;
