MU Replication: The property view now displays the selected object and its components. The properties are categorized under a category row.

diff --git a/Engine/Plugins/Developer/Concert/ConcertApp/MultiUserClient/Source/MultiUserClient/Private/Widgets/ActiveSession/Replication/Client/Multi/Columns/MultiStreamColumns_AssignProperties.cpp b/Engine/Plugins/Developer/Concert/ConcertApp/MultiUserClient/Source/MultiUserClient/Private/Widgets/ActiveSession/Replication/Client/Multi/Columns/MultiStreamColumns_AssignProperties.cpp
--- a/Engine/Plugins/Developer/Concert/ConcertApp/MultiUserClient/Source/MultiUserClient/Private/Widgets/ActiveSession/Replication/Client/Multi/Columns/MultiStreamColumns_AssignProperties.cpp
+++ b/Engine/Plugins/Developer/Concert/ConcertApp/MultiUserClient/Source/MultiUserClient/Private/Widgets/ActiveSession/Replication/Client/Multi/Columns/MultiStreamColumns_AssignProperties.cpp
@@ -106,7 +106,7 @@ namespace UE::MultiUserClient::MultiStreamColumns
 			
 			virtual TSharedRef<SWidget> GenerateColumnWidget(const FBuildArgs& InArgs) override
 			{
-				const TSet<FSoftObjectPath> DisplayedObjects = MultiStreamEditor.Get()->GetEditorBase().GetObjectsBeingPropertyEdited();
+				const TArray<FSoftObjectPath> DisplayedObjects = MultiStreamEditor.Get()->GetEditorBase().GetObjectsBeingPropertyEdited();
 				return SNew(SAssignPropertyComboBox, MultiStreamEditor.Get().ToSharedRef(), ConcertClient, ClientManager)
 					.DisplayedProperty(InArgs.RowItem.RowData.GetProperty())
 					.EditedObjects(DisplayedObjects)
@@ -132,7 +132,7 @@ namespace UE::MultiUserClient::MultiStreamColumns
 			virtual bool CanBeSorted() const override { return true; }
 			virtual bool IsLessThan(const FPropertyTreeRowContext& Left, const FPropertyTreeRowContext& Right) const override
 			{
-				const TSet<FSoftObjectPath> DisplayedObjects = MultiStreamEditor.Get()->GetEditorBase().GetObjectsBeingPropertyEdited();
+				const TArray<FSoftObjectPath> DisplayedObjects = MultiStreamEditor.Get()->GetEditorBase().GetObjectsBeingPropertyEdited();
 				const TOptional<FString> LeftClientDisplayString = SAssignPropertyComboBox::GetDisplayString(ConcertClient, ClientManager, Left.RowData.GetProperty(), DisplayedObjects);
 				const TOptional<FString> RightClientDisplayString = SAssignPropertyComboBox::GetDisplayString(ConcertClient, ClientManager, Right.RowData.GetProperty(), DisplayedObjects);
 			
diff --git a/Engine/Plugins/Developer/Concert/ConcertApp/MultiUserClient/Source/MultiUserClient/Private/Widgets/ActiveSession/Replication/Client/Multi/Columns/SAssignPropertyComboBox.cpp b/Engine/Plugins/Developer/Concert/ConcertApp/MultiUserClient/Source/MultiUserClient/Private/Widgets/ActiveSession/Replication/Client/Multi/Columns/SAssignPropertyComboBox.cpp
--- a/Engine/Plugins/Developer/Concert/ConcertApp/MultiUserClient/Source/MultiUserClient/Private/Widgets/ActiveSession/Replication/Client/Multi/Columns/SAssignPropertyComboBox.cpp
+++ b/Engine/Plugins/Developer/Concert/ConcertApp/MultiUserClient/Source/MultiUserClient/Private/Widgets/ActiveSession/Replication/Client/Multi/Columns/SAssignPropertyComboBox.cpp
@@ -30,7 +30,7 @@ namespace UE::MultiUserClient
 {
 	namespace AssignPropertyComboBox
 	{
-		TArray<FGuid> GetDisplayedClients(const FReplicationClientManager& ClientManager, const FConcertPropertyChain& DisplayedProperty, const TSet<FSoftObjectPath>& EditedObjects)
+		TArray<FGuid> GetDisplayedClients(const FReplicationClientManager& ClientManager, const FConcertPropertyChain& DisplayedProperty, const TArray<FSoftObjectPath>& EditedObjects)
 		{
 			TArray<FGuid> Clients;
 			ClientManager.ForEachClient([&DisplayedProperty, &EditedObjects, &Clients](const FReplicationClient& Client)
@@ -55,7 +55,7 @@ namespace UE::MultiUserClient
 		const TSharedRef<IConcertClient>& LocalConcertClient,
 		const FReplicationClientManager& ClientManager,
 		const FConcertPropertyChain& DisplayedProperty,
-		const TSet<FSoftObjectPath>& EditedObjects)
+		const TArray<FSoftObjectPath>& EditedObjects)
 	{
 		using SWidgetType = ConcertClientSharedSlate::SHorizontalClientList;
 		const TArray<FGuid> Clients = AssignPropertyComboBox::GetDisplayedClients(ClientManager, DisplayedProperty, EditedObjects);
diff --git a/Engine/Plugins/Developer/Concert/ConcertApp/MultiUserClient/Source/MultiUserClient/Private/Widgets/ActiveSession/Replication/Client/Multi/Columns/SAssignPropertyComboBox.h b/Engine/Plugins/Developer/Concert/ConcertApp/MultiUserClient/Source/MultiUserClient/Private/Widgets/ActiveSession/Replication/Client/Multi/Columns/SAssignPropertyComboBox.h
--- a/Engine/Plugins/Developer/Concert/ConcertApp/MultiUserClient/Source/MultiUserClient/Private/Widgets/ActiveSession/Replication/Client/Multi/Columns/SAssignPropertyComboBox.h
+++ b/Engine/Plugins/Developer/Concert/ConcertApp/MultiUserClient/Source/MultiUserClient/Private/Widgets/ActiveSession/Replication/Client/Multi/Columns/SAssignPropertyComboBox.h
@@ -38,13 +38,13 @@ namespace UE::MultiUserClient
 			const TSharedRef<IConcertClient>& LocalConcertClient,
 			const FReplicationClientManager& ClientManager,
 			const FConcertPropertyChain& DisplayedProperty,
-			const TSet<FSoftObjectPath>& EditedObjects
+			const TArray<FSoftObjectPath>& EditedObjects
 			);
 		
 		SLATE_BEGIN_ARGS(SAssignPropertyComboBox)
 		{}
 			SLATE_ARGUMENT(FConcertPropertyChain, DisplayedProperty)
-			SLATE_ARGUMENT(TSet<FSoftObjectPath>, EditedObjects)
+			SLATE_ARGUMENT(TArray<FSoftObjectPath>, EditedObjects)
 			SLATE_ARGUMENT(TSharedPtr<FText>, HighlightText)
 			/** Called when the property assignment of any client(s) is changed by this widget. */
 			SLATE_EVENT(FOnPropertyAssignmentChanged, OnPropertyAssignmentChanged)
@@ -65,7 +65,7 @@ namespace UE::MultiUserClient
 		TSharedPtr<IConcertClient> ConcertClient;
 
 		/** The objects for which the property is being displayed. */
-		TSet<FSoftObjectPath> EditedObjects;
+		TArray<FSoftObjectPath> EditedObjects;
 		/** The property assigned to this column. */
 		FConcertPropertyChain Property;
 
diff --git a/Engine/Plugins/Developer/Concert/ConcertApp/MultiUserClient/Source/MultiUserClient/Private/Widgets/ActiveSession/Replication/Client/Multi/SMultiClientView.cpp b/Engine/Plugins/Developer/Concert/ConcertApp/MultiUserClient/Source/MultiUserClient/Private/Widgets/ActiveSession/Replication/Client/Multi/SMultiClientView.cpp
--- a/Engine/Plugins/Developer/Concert/ConcertApp/MultiUserClient/Source/MultiUserClient/Private/Widgets/ActiveSession/Replication/Client/Multi/SMultiClientView.cpp
+++ b/Engine/Plugins/Developer/Concert/ConcertApp/MultiUserClient/Source/MultiUserClient/Private/Widgets/ActiveSession/Replication/Client/Multi/SMultiClientView.cpp
@@ -10,6 +10,7 @@
 #include "Replication/Client/ReplicationClient.h"
 #include "Replication/Client/ReplicationClientManager.h"
 #include "Replication/Editor/Model/PropertySource/SelectPropertyFromUClassModel.h"
+#include "Replication/Editor/View/IMultiObjectPropertyAssignmentView.h"
 #include "Replication/Editor/View/IMultiReplicationStreamEditor.h"
 #include "Replication/Editor/View/IReplicationStreamEditor.h"
 #include "Replication/Editor/Model/ObjectSource/ActorSelectionSourceModel.h"
@@ -78,6 +79,9 @@ namespace UE::MultiUserClient
 
 		FMuteStateManager& MuteManager = *InMultiUserReplicationManager.GetMuteManager();
 		
+		ObjectHierarchy = ConcertClientSharedSlate::CreateObjectHierarchyForComponentHierarchy();
+		const TSharedRef<IObjectNameModel> NameModel = ConcertClientSharedSlate::CreateEditorObjectNameModel();
+		
 		TAttribute<TSharedPtr<IMultiReplicationStreamEditor>> MultiStreamEditorAttribute =
 		   TAttribute<TSharedPtr<IMultiReplicationStreamEditor>>::CreateLambda([this]()
 		   {
@@ -101,10 +105,14 @@ namespace UE::MultiUserClient
 			{
 				ReplicationColumns::Property::LabelColumn(),
 				MultiStreamColumns::AssignPropertyColumn(MultiStreamEditorAttribute, InConcertClient, *ClientManager)
-			}
+			},
+			.CreateCategoryRow = CreateDefaultCategoryGenerator(NameModel)
 		};
 		TSharedRef<IPropertyTreeView> PropertyTreeView = CreateFilterablePropertyTreeView(MoveTemp(TreeViewParams));
-		TSharedRef<IPropertyAssignmentView> PropertyAssignmentView = CreatePerObjectAssignmentView({ .PropertyTreeView = PropertyTreeView, .PropertySource = PropertySourceModel });
+		
+		TSharedRef<IMultiObjectPropertyAssignmentView> PropertyAssignmentView = CreateMultiObjectAssignmentView(
+			{ .PropertyTreeView = PropertyTreeView, .ObjectHierarchy = ObjectHierarchy, .PropertySource = PropertySourceModel}
+			);
 		
 		FCreateMultiStreamEditorParams Params
 		{
@@ -115,12 +123,11 @@ namespace UE::MultiUserClient
 			.GetAutoAssignToStreamDelegate = MoveTemp(GetAutoAssignTargetDelegate)
 		};
 		
-		ObjectHierarchy = ConcertClientSharedSlate::CreateObjectHierarchyForComponentHierarchy();
 		FCreateViewerParams ViewerParams
 		{
-			.PropertyAssignmentView = MoveTemp(PropertyAssignmentView),
+			.PropertyAssignmentView = StaticCastSharedRef<IPropertyAssignmentView>(PropertyAssignmentView),
 			.ObjectHierarchy = ObjectHierarchy, // This makes actors have children in the top view
-			.NameModel = ConcertClientSharedSlate::CreateEditorObjectNameModel(), // This makes actors use their labels, and components use the names given in the BP editor
+			.NameModel = NameModel, // This makes actors use their labels, and components use the names given in the BP editor
 			.OnExtendObjectsContextMenu = FExtendObjectMenu::CreateSP(this, &SMultiClientView::ExtendObjectContextMenu),
 			.ObjectColumns =
 			{
diff --git a/Engine/Plugins/Developer/Concert/ConcertUI/ConcertClientSharedSlate/Source/ConcertClientSharedSlate/Private/Replication/Editor/View/PropertyTree/SFilteredPropertyTreeView.cpp b/Engine/Plugins/Developer/Concert/ConcertUI/ConcertClientSharedSlate/Source/ConcertClientSharedSlate/Private/Replication/Editor/View/PropertyTree/SFilteredPropertyTreeView.cpp
--- a/Engine/Plugins/Developer/Concert/ConcertUI/ConcertClientSharedSlate/Source/ConcertClientSharedSlate/Private/Replication/Editor/View/PropertyTree/SFilteredPropertyTreeView.cpp
+++ b/Engine/Plugins/Developer/Concert/ConcertUI/ConcertClientSharedSlate/Source/ConcertClientSharedSlate/Private/Replication/Editor/View/PropertyTree/SFilteredPropertyTreeView.cpp
@@ -62,6 +62,7 @@ namespace UE::ConcertClientSharedSlate
 		{
 			.PropertyColumns = MoveTemp(Params.AdditionalPropertyColumns),
 			.FilterItem = FFilterPropertyData::CreateSP(this, &SFilteredPropertyTreeView::PassesFilters),
+			.CreateCategoryRow = MoveTemp(Params.CreateCategoryRow),
 			.PrimaryPropertySort = Params.PrimaryPropertySort,
 			.SecondaryPropertySort = Params.SecondaryPropertySort
 		};
diff --git a/Engine/Plugins/Developer/Concert/ConcertUI/ConcertClientSharedSlate/Source/ConcertClientSharedSlate/Public/Replication/ClientReplicationWidgetFactories.h b/Engine/Plugins/Developer/Concert/ConcertUI/ConcertClientSharedSlate/Source/ConcertClientSharedSlate/Public/Replication/ClientReplicationWidgetFactories.h
--- a/Engine/Plugins/Developer/Concert/ConcertUI/ConcertClientSharedSlate/Source/ConcertClientSharedSlate/Public/Replication/ClientReplicationWidgetFactories.h
+++ b/Engine/Plugins/Developer/Concert/ConcertUI/ConcertClientSharedSlate/Source/ConcertClientSharedSlate/Public/Replication/ClientReplicationWidgetFactories.h
@@ -2,6 +2,7 @@
 
 #pragma once
 
+#include "Replication/PropertyTreeFactory.h"
 #include "Replication/ReplicationWidgetFactories.h"
 #include "Templates/SharedPointer.h"
 
@@ -56,6 +57,15 @@ namespace UE::ConcertClientSharedSlate
 		ConcertSharedSlate::FColumnSortInfo PrimaryPropertySort { ConcertSharedSlate::ReplicationColumns::Property::LabelColumnId, EColumnSortMode::Ascending };
 		/** Optional initial secondary sort mode for object rows */
 		ConcertSharedSlate::FColumnSortInfo SecondaryPropertySort { ConcertSharedSlate::ReplicationColumns::Property::LabelColumnId, EColumnSortMode::Ascending };
+		
+		/**
+		 * Optional delegate for grouping objects under a category.
+		 * If unset, no category are generated.
+		 * 
+		 * When the user clicks an object in the top view, this delegate will be called for the clicked object, its components (if an actor), and its (nested) subobjects.
+		 * ContextObjects is a single object if a single object is clicked or multiple object in the case of multi-edit.
+		 */
+		ConcertSharedSlate::FCreateCategoryRow CreateCategoryRow;
 	};
 	
 	/**
diff --git a/Engine/Plugins/Developer/Concert/ConcertUI/ConcertSharedSlate/Source/ConcertSharedSlate/Private/Replication/Editor/Model/Data/PropertyNodeData.h b/Engine/Plugins/Developer/Concert/ConcertUI/ConcertSharedSlate/Source/ConcertSharedSlate/Private/Replication/Editor/Model/Data/PropertyNodeData.h
new file mode 100644
--- /dev/null
+++ b/Engine/Plugins/Developer/Concert/ConcertUI/ConcertSharedSlate/Source/ConcertSharedSlate/Private/Replication/Editor/Model/Data/PropertyNodeData.h
@@ -0,0 +1,50 @@
+﻿// Copyright Epic Games, Inc. All Rights Reserved.
+
+#pragma once
+
+#include "Replication/Editor/Model/Data/PropertyData.h"
+
+#include "Containers/Array.h"
+#include "UObject/SoftObjectPath.h"
+#include "UObject/SoftObjectPtr.h"
+
+namespace UE::ConcertSharedSlate
+{
+	/** A node in the tree view which groups a subobject. */
+	class FCategoryData
+	{
+	public:
+
+		FCategoryData(TArray<TSoftObjectPtr<>> InContextObjects) : ContextObjects(MoveTemp(InContextObjects)) {}
+
+		const TArray<TSoftObjectPtr<>>& GetContextObjects() const { return ContextObjects; }
+
+	private:
+
+		/** The objects this category represents; all the objects are releated, i.e. ActorA's StaticMeshComponent0 and ActorB's StaticMeshComponent0. */
+		const TArray<TSoftObjectPtr<>> ContextObjects;
+	};
+	
+	/** Instanced for each property row in IPropertyTreeView.*/
+	class FPropertyNodeData
+	{
+	public:
+		
+		FPropertyNodeData(FPropertyData PropertyData)
+			: PropertyData(MoveTemp(PropertyData))
+		{}
+		FPropertyNodeData(FCategoryData CategoryData)
+			: CategoryData(MoveTemp(CategoryData))
+		{}
+		
+		bool IsCategoryNode() const { return CategoryData.IsSet() && ensure(!PropertyData); }
+		
+		const TOptional<FPropertyData>& GetPropertyData() const { return PropertyData; }
+		const TOptional<FCategoryData>& GetCategoryData() const { return CategoryData; }
+
+	private:
+		
+		const TOptional<FPropertyData> PropertyData;
+		const TOptional<FCategoryData> CategoryData;
+	};
+}
\ No newline at end of file
diff --git a/Engine/Plugins/Developer/Concert/ConcertUI/ConcertSharedSlate/Source/ConcertSharedSlate/Private/Replication/Editor/Model/Object/IObjectHierarchyModel.cpp b/Engine/Plugins/Developer/Concert/ConcertUI/ConcertSharedSlate/Source/ConcertSharedSlate/Private/Replication/Editor/Model/Object/IObjectHierarchyModel.cpp
--- a/Engine/Plugins/Developer/Concert/ConcertUI/ConcertSharedSlate/Source/ConcertSharedSlate/Private/Replication/Editor/Model/Object/IObjectHierarchyModel.cpp
+++ b/Engine/Plugins/Developer/Concert/ConcertUI/ConcertSharedSlate/Source/ConcertSharedSlate/Private/Replication/Editor/Model/Object/IObjectHierarchyModel.cpp
@@ -33,7 +33,7 @@ namespace UE::ConcertSharedSlate
 				{
 					Result = Visit(Model, ChildObject, ChildOfChild, Relationship, Callback, InclusionFlags);
 					return Result;
-				});
+				}, InclusionFlags);
 				return Result;
 			}
 		};
diff --git a/Engine/Plugins/Developer/Concert/ConcertUI/ConcertSharedSlate/Source/ConcertSharedSlate/Private/Replication/Editor/Utils/PropertyEnumerationUtils.cpp b/Engine/Plugins/Developer/Concert/ConcertUI/ConcertSharedSlate/Source/ConcertSharedSlate/Private/Replication/Editor/Utils/PropertyEnumerationUtils.cpp
--- a/Engine/Plugins/Developer/Concert/ConcertUI/ConcertSharedSlate/Source/ConcertSharedSlate/Private/Replication/Editor/Utils/PropertyEnumerationUtils.cpp
+++ b/Engine/Plugins/Developer/Concert/ConcertUI/ConcertSharedSlate/Source/ConcertSharedSlate/Private/Replication/Editor/Utils/PropertyEnumerationUtils.cpp
@@ -29,7 +29,7 @@ namespace UE::ConcertSharedSlate
 #endif
 	}
 
-	void EnumerateProperties(const TSet<FSoftObjectPath>& Objects, const IReplicationStreamModel& Model, const IPropertySelectionSourceModel* OptionalSource, FEnumerateProperties Callback)
+	void EnumerateProperties(TConstArrayView<FSoftObjectPath> Objects, const IReplicationStreamModel& Model, const IPropertySelectionSourceModel* OptionalSource, FEnumerateProperties Callback)
 	{
 		if (OptionalSource)
 		{
@@ -42,7 +42,7 @@ namespace UE::ConcertSharedSlate
 	}
 
 	/** Enumerates the properties that are assigned to the object in Model */
-	void EnumerateRegisteredPropertiesOnly(const TSet<FSoftObjectPath>& Objects, const IReplicationStreamModel& Model, FEnumerateProperties Callback)
+	void EnumerateRegisteredPropertiesOnly(TConstArrayView<FSoftObjectPath> Objects, const IReplicationStreamModel& Model, FEnumerateProperties Callback)
 	{
 		for (const FSoftObjectPath& Object : Objects)
 		{
@@ -63,7 +63,7 @@ namespace UE::ConcertSharedSlate
 	}
 
 	/** Enumerate the properties that are selectable in Source (e.g. all properties in that class, @see FSelectPropertyFromUClassModel). */
-	void EnumerateAllProperties(const TSet<FSoftObjectPath>& Objects, const IPropertySelectionSourceModel& Source, const IReplicationStreamModel& Model, FEnumerateProperties Callback)
+	void EnumerateAllProperties(TConstArrayView<FSoftObjectPath> Objects, const IPropertySelectionSourceModel& Source, const IReplicationStreamModel& Model, FEnumerateProperties Callback)
 	{
 		TSet<FSoftClassPath> VisitedClasses; 
 			
diff --git a/Engine/Plugins/Developer/Concert/ConcertUI/ConcertSharedSlate/Source/ConcertSharedSlate/Private/Replication/Editor/Utils/PropertyEnumerationUtils.h b/Engine/Plugins/Developer/Concert/ConcertUI/ConcertSharedSlate/Source/ConcertSharedSlate/Private/Replication/Editor/Utils/PropertyEnumerationUtils.h
--- a/Engine/Plugins/Developer/Concert/ConcertUI/ConcertSharedSlate/Source/ConcertSharedSlate/Private/Replication/Editor/Utils/PropertyEnumerationUtils.h
+++ b/Engine/Plugins/Developer/Concert/ConcertUI/ConcertSharedSlate/Source/ConcertSharedSlate/Private/Replication/Editor/Utils/PropertyEnumerationUtils.h
@@ -22,9 +22,9 @@ namespace UE::ConcertSharedSlate
 	FSoftClassPath GetObjectClassFromModelOrLoad(const FSoftObjectPath& Object, const IReplicationStreamModel& Model);
 
 	/** Calls EnumerateRegisteredPropertiesOnly or EnumerateAllProperties depending on whether OptionalSource is nullptr. */
-	void EnumerateProperties(const TSet<FSoftObjectPath>& Objects, const IReplicationStreamModel& Model, const IPropertySelectionSourceModel* OptionalSource, FEnumerateProperties Callback);
+	void EnumerateProperties(TConstArrayView<FSoftObjectPath> Objects, const IReplicationStreamModel& Model, const IPropertySelectionSourceModel* OptionalSource, FEnumerateProperties Callback);
 	/** Enumerates the properties that are assigned to the object in Model */
-	void EnumerateRegisteredPropertiesOnly(const TSet<FSoftObjectPath>& Objects, const IReplicationStreamModel& Model, FEnumerateProperties Callback);
+	void EnumerateRegisteredPropertiesOnly(TConstArrayView<FSoftObjectPath> Objects, const IReplicationStreamModel& Model, FEnumerateProperties Callback);
 	/** Enumerate the properties that are selectable in Source (e.g. all properties in that class, @see FSelectPropertyFromUClassModel). */
-	void EnumerateAllProperties(const TSet<FSoftObjectPath>& Objects, const IPropertySelectionSourceModel& Source, const IReplicationStreamModel& Model, FEnumerateProperties Callback);
+	void EnumerateAllProperties(TConstArrayView<FSoftObjectPath> Objects, const IPropertySelectionSourceModel& Source, const IReplicationStreamModel& Model, FEnumerateProperties Callback);
 }
\ No newline at end of file
diff --git a/Engine/Plugins/Developer/Concert/ConcertUI/ConcertSharedSlate/Source/ConcertSharedSlate/Private/Replication/Editor/View/Column/PropertyColumnAdapter.h b/Engine/Plugins/Developer/Concert/ConcertUI/ConcertSharedSlate/Source/ConcertSharedSlate/Private/Replication/Editor/View/Column/PropertyColumnAdapter.h
--- a/Engine/Plugins/Developer/Concert/ConcertUI/ConcertSharedSlate/Source/ConcertSharedSlate/Private/Replication/Editor/View/Column/PropertyColumnAdapter.h
+++ b/Engine/Plugins/Developer/Concert/ConcertUI/ConcertSharedSlate/Source/ConcertSharedSlate/Private/Replication/Editor/View/Column/PropertyColumnAdapter.h
@@ -3,23 +3,24 @@
 #pragma once
 
 #include "Replication/Editor/Model/Data/PropertyData.h"
+#include "Replication/Editor/Model/Data/PropertyNodeData.h"
 #include "Replication/Editor/View/Column/IPropertyTreeColumn.h"
 #include "Replication/Editor/View/Column/IReplicationTreeColumn.h"
 
 namespace UE::ConcertSharedSlate
 {
 	/** Adapts an IPropertyTreeColumn to IReplicationTreeColumn<FPropertyData>. It simply passes additional info down to IPropertyTreeColumn. */
-	class FPropertyColumnAdapter : public IReplicationTreeColumn<FPropertyData>
+	class FPropertyColumnAdapter : public IReplicationTreeColumn<FPropertyNodeData>
 	{
 	public:
 
-		static TArray<TReplicationColumnEntry<FPropertyData>> Transform(const TArray<FPropertyColumnEntry>& Entries)
+		static TArray<TReplicationColumnEntry<FPropertyNodeData>> Transform(const TArray<FPropertyColumnEntry>& Entries)
 		{
-			TArray<TReplicationColumnEntry<FPropertyData>> Result;
-			Algo::Transform(Entries, Result, [](const FPropertyColumnEntry& Entry) -> TReplicationColumnEntry<FPropertyData>
+			TArray<TReplicationColumnEntry<FPropertyNodeData>> Result;
+			Algo::Transform(Entries, Result, [](const FPropertyColumnEntry& Entry) -> TReplicationColumnEntry<FPropertyNodeData>
 			{
 				return {
-					TReplicationColumnDelegates<FPropertyData>::FCreateColumn::CreateLambda([CreateDelegate = Entry.CreateColumn]()
+					TReplicationColumnDelegates<FPropertyNodeData>::FCreateColumn::CreateLambda([CreateDelegate = Entry.CreateColumn]()
 					{
 						return MakeShared<FPropertyColumnAdapter>(CreateDelegate.Execute());
 					}),
@@ -39,21 +40,23 @@ namespace UE::ConcertSharedSlate
 		{
 			return AdaptedColumn->GenerateColumnWidget({ InArgs.HighlightText, Transform(InArgs.RowItem) });
 		}
-		virtual void PopulateSearchString(const FPropertyData& InItem, TArray<FString>& InOutSearchStrings) const override
+		virtual void PopulateSearchString(const FPropertyNodeData& InItem, TArray<FString>& InOutSearchStrings) const override
 		{
 			return AdaptedColumn->PopulateSearchString(Transform(InItem), InOutSearchStrings);
 		}
 
 		virtual bool CanBeSorted() const override { return AdaptedColumn->CanBeSorted(); }
-		virtual bool IsLessThan(const FPropertyData& Left, const FPropertyData& Right) const override { return AdaptedColumn->IsLessThan(Transform(Left), Transform(Right)); }
+		virtual bool IsLessThan(const FPropertyNodeData& Left, const FPropertyNodeData& Right) const override { return AdaptedColumn->IsLessThan(Transform(Left), Transform(Right)); }
 
 	private:
 
 		TSharedRef<IPropertyTreeColumn> AdaptedColumn;
 
-		static FPropertyTreeRowContext Transform(FPropertyData Data)
+		static FPropertyTreeRowContext Transform(const FPropertyNodeData& Data)
 		{
-			return { MoveTemp(Data) };
+			const TOptional<FPropertyData>& PropertyData = Data.GetPropertyData();
+			checkf(PropertyData, TEXT("This node is not a property node - column callbacks should not have been invoked on it!"));
+			return { *PropertyData };
 		}
 	};
 }
diff --git a/Engine/Plugins/Developer/Concert/ConcertUI/ConcertSharedSlate/Source/ConcertSharedSlate/Private/Replication/Editor/View/Extension/CategoryRowGeneration.cpp b/Engine/Plugins/Developer/Concert/ConcertUI/ConcertSharedSlate/Source/ConcertSharedSlate/Private/Replication/Editor/View/Extension/CategoryRowGeneration.cpp
new file mode 100644
--- /dev/null
+++ b/Engine/Plugins/Developer/Concert/ConcertUI/ConcertSharedSlate/Source/ConcertSharedSlate/Private/Replication/Editor/View/Extension/CategoryRowGeneration.cpp
@@ -0,0 +1,68 @@
+﻿// Copyright Epic Games, Inc. All Rights Reserved.
+
+#include "Replication/Editor/View/Extension/CategoryRowGeneration.h"
+
+#include "Replication/Editor/Model/Object/IObjectNameModel.h"
+
+#include "Widgets/DeclarativeSyntaxSupport.h"
+#include "Widgets/SCompoundWidget.h"
+#include "Widgets/Text/STextBlock.h"
+
+namespace UE::ConcertSharedSlate
+{
+	namespace Private
+	{
+		class SDefaultCategoryRow : public SCompoundWidget, public ICategoryRow
+		{
+			FText Label;
+		public:
+
+			SLATE_BEGIN_ARGS(SDefaultCategoryRow){}
+				SLATE_ARGUMENT(FText, Label)
+				SLATE_ATTRIBUTE(FText, HighlightText)
+			SLATE_END_ARGS()
+
+			void Construct(const FArguments& InArgs)
+			{
+				Label = InArgs._Label;
+				ChildSlot
+				[
+					SNew(STextBlock)
+					.Text(Label)
+					.HighlightText(InArgs._HighlightText)
+					.Font(FAppStyle::Get().GetFontStyle("DetailsView.CategoryFontStyle"))
+					.TextStyle(FAppStyle::Get(), "DetailsView.CategoryTextStyle")
+				];
+			}
+
+			//~ Begin ICategoryRow Interface
+			virtual void GenerateSearchTerms(const TArray<TSoftObjectPtr<>>& ContextObjects, TArray<FString>& InOutSearchTerms) const override
+			{
+				const FString LabelString = Label.ToString();
+				InOutSearchTerms.Add(LabelString);
+				
+				TArray<FString> Split;
+				LabelString.ParseIntoArray(Split, TEXT(" "));
+				InOutSearchTerms.Append(Split);
+			}
+			virtual TSharedRef<SWidget> GetWidget() override { return SharedThis(this); }
+			//~ End ICategoryRow Interface
+		};
+	}
+	
+	FCreateCategoryRow CreateDefaultCategoryGenerator(TSharedRef<IObjectNameModel> NameModel)
+	{
+		return FCreateCategoryRow::CreateLambda([NameModel](const FCategoryRowGenerationArgs& Args) -> TSharedRef<ICategoryRow>
+		{
+			const TArray<TSoftObjectPtr<>>& ContextObjects = Args.ContextObjects;
+			const FText Label = ContextObjects.IsEmpty()
+				? FText::GetEmpty()
+				: NameModel->GetObjectDisplayName(ContextObjects[0]);
+			
+			return SNew(Private::SDefaultCategoryRow)
+				.Label(Label)
+				.HighlightText(Args.HighlightText);
+		});
+	}
+}
+
diff --git a/Engine/Plugins/Developer/Concert/ConcertUI/ConcertSharedSlate/Source/ConcertSharedSlate/Private/Replication/Editor/View/ObjectEditor/SBaseReplicationStreamEditor.cpp b/Engine/Plugins/Developer/Concert/ConcertUI/ConcertSharedSlate/Source/ConcertSharedSlate/Private/Replication/Editor/View/ObjectEditor/SBaseReplicationStreamEditor.cpp
--- a/Engine/Plugins/Developer/Concert/ConcertUI/ConcertSharedSlate/Source/ConcertSharedSlate/Private/Replication/Editor/View/ObjectEditor/SBaseReplicationStreamEditor.cpp
+++ b/Engine/Plugins/Developer/Concert/ConcertUI/ConcertSharedSlate/Source/ConcertSharedSlate/Private/Replication/Editor/View/ObjectEditor/SBaseReplicationStreamEditor.cpp
@@ -93,7 +93,7 @@ namespace UE::ConcertSharedSlate
 		ReplicationViewer->RequestPropertyColumnResort(ColumnId);
 	}
 
-	TSet<FSoftObjectPath> SBaseReplicationStreamEditor::GetObjectsBeingPropertyEdited() const
+	TArray<FSoftObjectPath> SBaseReplicationStreamEditor::GetObjectsBeingPropertyEdited() const
 	{
 		return ReplicationViewer->GetObjectsBeingPropertyEdited();
 	}
diff --git a/Engine/Plugins/Developer/Concert/ConcertUI/ConcertSharedSlate/Source/ConcertSharedSlate/Private/Replication/Editor/View/ObjectEditor/SBaseReplicationStreamEditor.h b/Engine/Plugins/Developer/Concert/ConcertUI/ConcertSharedSlate/Source/ConcertSharedSlate/Private/Replication/Editor/View/ObjectEditor/SBaseReplicationStreamEditor.h
--- a/Engine/Plugins/Developer/Concert/ConcertUI/ConcertSharedSlate/Source/ConcertSharedSlate/Private/Replication/Editor/View/ObjectEditor/SBaseReplicationStreamEditor.h
+++ b/Engine/Plugins/Developer/Concert/ConcertUI/ConcertSharedSlate/Source/ConcertSharedSlate/Private/Replication/Editor/View/ObjectEditor/SBaseReplicationStreamEditor.h
@@ -97,7 +97,7 @@ namespace UE::ConcertSharedSlate
 		virtual void Refresh() override;
 		virtual void RequestObjectColumnResort(const FName& ColumnId) override;
 		virtual void RequestPropertyColumnResort(const FName& ColumnId) override;
-		virtual TSet<FSoftObjectPath> GetObjectsBeingPropertyEdited() const override;
+		virtual TArray<FSoftObjectPath> GetObjectsBeingPropertyEdited() const override;
 		//~ End IReplicationStreamEditor Interface
 
 	private:
diff --git a/Engine/Plugins/Developer/Concert/ConcertUI/ConcertSharedSlate/Source/ConcertSharedSlate/Private/Replication/Editor/View/ObjectViewer/SReplicatedPropertyView.cpp b/Engine/Plugins/Developer/Concert/ConcertUI/ConcertSharedSlate/Source/ConcertSharedSlate/Private/Replication/Editor/View/ObjectViewer/SReplicatedPropertyView.cpp
--- a/Engine/Plugins/Developer/Concert/ConcertUI/ConcertSharedSlate/Source/ConcertSharedSlate/Private/Replication/Editor/View/ObjectViewer/SReplicatedPropertyView.cpp
+++ b/Engine/Plugins/Developer/Concert/ConcertUI/ConcertSharedSlate/Source/ConcertSharedSlate/Private/Replication/Editor/View/ObjectViewer/SReplicatedPropertyView.cpp
@@ -32,7 +32,7 @@ namespace UE::ConcertSharedSlate
 	void SReplicatedPropertyView::RefreshPropertyData()
 	{
 		SCOPED_CONCERT_TRACE(RefreshPropertyData);
-		TSet<FSoftObjectPath> SelectedObjects = GetObjectsSelectedForPropertyEditing();
+		TArray<FSoftObjectPath> SelectedObjects = GetObjectsSelectedForPropertyEditing();
 		if (SelectedObjects.IsEmpty())
 		{
 			SetPropertyContent(EReplicatedPropertyContent::NoSelection);
@@ -51,9 +51,9 @@ namespace UE::ConcertSharedSlate
 		SetPropertyContent(EReplicatedPropertyContent::Properties);
 	}
 
-	TSet<FSoftObjectPath> SReplicatedPropertyView::GetObjectsSelectedForPropertyEditing() const
+	TArray<FSoftObjectPath> SReplicatedPropertyView::GetObjectsSelectedForPropertyEditing() const
 	{
-		TSet<FSoftObjectPath> Result;
+		TArray<FSoftObjectPath> Result;
 		Algo::Transform(GetSelectedRootObjectsDelegate.Execute(), Result, [](const TSharedPtr<FReplicatedObjectData>& ObjectData)
 		{
 			return ObjectData->GetObjectPath();
@@ -92,7 +92,7 @@ namespace UE::ConcertSharedSlate
 			];
 	}
 	
-	TOptional<FSoftClassPath> SReplicatedPropertyView::GetClassForPropertiesFromSelection(const TSet<FSoftObjectPath>& Objects) const
+	TOptional<FSoftClassPath> SReplicatedPropertyView::GetClassForPropertiesFromSelection(const TArray<FSoftObjectPath>& Objects) const
 	{
 		FSoftClassPath SharedClass;
 		const bool bAllHaveSameClass = Algo::AllOf(Objects, [this, &SharedClass](const FSoftObjectPath& Object)
diff --git a/Engine/Plugins/Developer/Concert/ConcertUI/ConcertSharedSlate/Source/ConcertSharedSlate/Private/Replication/Editor/View/ObjectViewer/SReplicatedPropertyView.h b/Engine/Plugins/Developer/Concert/ConcertUI/ConcertSharedSlate/Source/ConcertSharedSlate/Private/Replication/Editor/View/ObjectViewer/SReplicatedPropertyView.h
--- a/Engine/Plugins/Developer/Concert/ConcertUI/ConcertSharedSlate/Source/ConcertSharedSlate/Private/Replication/Editor/View/ObjectViewer/SReplicatedPropertyView.h
+++ b/Engine/Plugins/Developer/Concert/ConcertUI/ConcertSharedSlate/Source/ConcertSharedSlate/Private/Replication/Editor/View/ObjectViewer/SReplicatedPropertyView.h
@@ -42,7 +42,7 @@ namespace UE::ConcertSharedSlate
 		/** Requests that the given column be resorted, if it currently affects the row sorting. */
 		void RequestResortForColumn(const FName& ColumnId) const { PropertyAssignmentView->RequestResortForColumn(ColumnId); }
 		
-		TSet<FSoftObjectPath> GetObjectsSelectedForPropertyEditing() const;
+		TArray<FSoftObjectPath> GetObjectsSelectedForPropertyEditing() const;
 
 	private:
 		
@@ -72,7 +72,7 @@ namespace UE::ConcertSharedSlate
 		TSharedRef<SWidget> CreatePropertiesView(const FArguments& InArgs);
 		
 		/** Given the selected objects, determines whether they all have the same class and returns it if so. */
-		TOptional<FSoftClassPath> GetClassForPropertiesFromSelection(const TSet<FSoftObjectPath>& Objects) const;
+		TOptional<FSoftClassPath> GetClassForPropertiesFromSelection(const TArray<FSoftObjectPath>& Objects) const;
 		/** Sets how to display this widget */
 		void SetPropertyContent(EReplicatedPropertyContent Content) const;
 
diff --git a/Engine/Plugins/Developer/Concert/ConcertUI/ConcertSharedSlate/Source/ConcertSharedSlate/Private/Replication/Editor/View/ObjectViewer/SReplicationStreamViewer.cpp b/Engine/Plugins/Developer/Concert/ConcertUI/ConcertSharedSlate/Source/ConcertSharedSlate/Private/Replication/Editor/View/ObjectViewer/SReplicationStreamViewer.cpp
--- a/Engine/Plugins/Developer/Concert/ConcertUI/ConcertSharedSlate/Source/ConcertSharedSlate/Private/Replication/Editor/View/ObjectViewer/SReplicationStreamViewer.cpp
+++ b/Engine/Plugins/Developer/Concert/ConcertUI/ConcertSharedSlate/Source/ConcertSharedSlate/Private/Replication/Editor/View/ObjectViewer/SReplicationStreamViewer.cpp
@@ -58,7 +58,7 @@ namespace UE::ConcertSharedSlate
 		PropertySection->RequestResortForColumn(ColumnId);
 	}
 
-	TSet<FSoftObjectPath> SReplicationStreamViewer::GetObjectsBeingPropertyEdited() const
+	TArray<FSoftObjectPath> SReplicationStreamViewer::GetObjectsBeingPropertyEdited() const
 	{
 		return PropertySection->GetObjectsSelectedForPropertyEditing();
 	}
diff --git a/Engine/Plugins/Developer/Concert/ConcertUI/ConcertSharedSlate/Source/ConcertSharedSlate/Private/Replication/Editor/View/ObjectViewer/SReplicationStreamViewer.h b/Engine/Plugins/Developer/Concert/ConcertUI/ConcertSharedSlate/Source/ConcertSharedSlate/Private/Replication/Editor/View/ObjectViewer/SReplicationStreamViewer.h
--- a/Engine/Plugins/Developer/Concert/ConcertUI/ConcertSharedSlate/Source/ConcertSharedSlate/Private/Replication/Editor/View/ObjectViewer/SReplicationStreamViewer.h
+++ b/Engine/Plugins/Developer/Concert/ConcertUI/ConcertSharedSlate/Source/ConcertSharedSlate/Private/Replication/Editor/View/ObjectViewer/SReplicationStreamViewer.h
@@ -95,7 +95,7 @@ namespace UE::ConcertSharedSlate
 		virtual void Refresh() override;
 		virtual void RequestObjectColumnResort(const FName& ColumnId) override;
 		virtual void RequestPropertyColumnResort(const FName& ColumnId) override;
-		virtual TSet<FSoftObjectPath> GetObjectsBeingPropertyEdited() const override;
+		virtual TArray<FSoftObjectPath> GetObjectsBeingPropertyEdited() const override;
 		//~ End IReplicationStreamViewer Interface
 
 		void RequestObjectDataRefresh() { bHasRequestedObjectRefresh = true; }
diff --git a/Engine/Plugins/Developer/Concert/ConcertUI/ConcertSharedSlate/Source/ConcertSharedSlate/Private/Replication/Editor/View/Property/SCategoryColumnRow.cpp b/Engine/Plugins/Developer/Concert/ConcertUI/ConcertSharedSlate/Source/ConcertSharedSlate/Private/Replication/Editor/View/Property/SCategoryColumnRow.cpp
new file mode 100644
--- /dev/null
+++ b/Engine/Plugins/Developer/Concert/ConcertUI/ConcertSharedSlate/Source/ConcertSharedSlate/Private/Replication/Editor/View/Property/SCategoryColumnRow.cpp
@@ -0,0 +1,89 @@
+﻿// Copyright Epic Games, Inc. All Rights Reserved.
+
+#include "SCategoryColumnRow.h"
+
+#include "Delegates/DelegateCombinations.h"
+#include "Widgets/Layout/SBox.h"
+#include "Widgets/SOverlay.h"
+#include "Widgets/Views/SExpanderArrow.h"
+#include "Widgets/Views/SListView.h"
+#include "Widgets/Views/STableRow.h"
+#include "Widgets/Views/STableViewBase.h"
+#include "Styling/AppStyle.h"
+#include "Styling/CoreStyle.h"
+
+namespace UE::ConcertSharedSlate
+{
+	void SCategoryColumnRow::Construct(const FArguments& InArgs, TSharedRef<STableViewBase> InOwner)
+	{
+		const FLinearColor BackgroundTint(0.6f, 0.6f, 0.6f, 1.0f);
+
+		// rebuilds the whole table row from scratch
+		ChildSlot
+		.Padding(0.0f, 2.0f, 0.0f, 0.0f)
+		[
+			SNew(SBox)
+			.MaxDesiredHeight(FConcertFrontendStyle::Get()->GetFloat("Concert.Replication.Tree.RowHeight"))
+			[
+				SNew(SBorder)
+				.BorderImage(this, &SCategoryColumnRow::GetBackgroundImage)
+				.Padding(FMargin(0.0f, 3.0f))
+				[
+					SNew(SHorizontalBox)
+					+SHorizontalBox::Slot()
+					.VAlign(VAlign_Center)
+					.Padding(2.0f, 2.0f, 2.0f, 2.0f)
+					.AutoWidth()
+					[
+						SNew(SExpanderArrow, SCategoryColumnRow::SharedThis(this))
+					]
+
+					+SHorizontalBox::Slot()
+					.VAlign(VAlign_Center)
+					[
+						InArgs._Content.Widget
+					]
+				]
+			]
+		];
+
+		ConstructInternal(
+			Super::FArguments()
+				.Style(FAppStyle::Get(), "DetailsView.TreeView.TableRow")
+				.ShowSelection(false),
+			InOwner
+		);
+	}
+
+	FReply SCategoryColumnRow::OnMouseButtonDown(const FGeometry& MyGeometry, const FPointerEvent& MouseEvent)
+	{
+		if (MouseEvent.GetEffectingButton() == EKeys::LeftMouseButton)
+		{
+			ToggleExpansion();
+			return FReply::Handled();
+		}
+
+		return FReply::Unhandled();
+	}
+
+	FReply SCategoryColumnRow::OnMouseButtonDoubleClick(const FGeometry& InMyGeometry, const FPointerEvent& InMouseEvent)
+	{
+		return OnMouseButtonDown(InMyGeometry, InMouseEvent);
+	}
+
+	const FSlateBrush* SCategoryColumnRow::GetBackgroundImage() const
+	{
+		if (IsHovered())
+		{
+			return IsItemExpanded()
+				? FAppStyle::GetBrush("DetailsView.CategoryTop_Hovered")
+				: FAppStyle::GetBrush("DetailsView.CollapsedCategory_Hovered");
+		}
+		else
+		{
+			return IsItemExpanded()
+				? FAppStyle::GetBrush("DetailsView.CategoryTop")
+				: FAppStyle::GetBrush("DetailsView.CollapsedCategory");
+		}
+	}
+}
diff --git a/Engine/Plugins/Developer/Concert/ConcertUI/ConcertSharedSlate/Source/ConcertSharedSlate/Private/Replication/Editor/View/Property/SCategoryColumnRow.h b/Engine/Plugins/Developer/Concert/ConcertUI/ConcertSharedSlate/Source/ConcertSharedSlate/Private/Replication/Editor/View/Property/SCategoryColumnRow.h
new file mode 100644
--- /dev/null
+++ b/Engine/Plugins/Developer/Concert/ConcertUI/ConcertSharedSlate/Source/ConcertSharedSlate/Private/Replication/Editor/View/Property/SCategoryColumnRow.h
@@ -0,0 +1,48 @@
+﻿// Copyright Epic Games, Inc. All Rights Reserved.
+
+#pragma once
+
+#include "ConcertFrontendStyle.h"
+
+#include "Delegates/DelegateCombinations.h"
+#include "Widgets/Layout/SBox.h"
+#include "Widgets/SNullWidget.h"
+#include "Widgets/SOverlay.h"
+#include "Widgets/Views/SExpanderArrow.h"
+#include "Widgets/Views/SListView.h"
+#include "Widgets/Views/STableRow.h"
+#include "Widgets/Views/STableViewBase.h"
+#include "Styling/AppStyle.h"
+#include "Styling/CoreStyle.h"
+#include "Replication/Editor/View/Column/IReplicationTreeColumn.h"
+
+namespace UE::ConcertSharedSlate
+{
+	class FPropertyNodeData;
+}
+
+namespace UE::ConcertSharedSlate
+{
+	/** A row that generates its columns generically via IReplicationTreeColumn. */
+	class SCategoryColumnRow : public STableRow<TSharedPtr<FPropertyNodeData>>
+	{
+		using Super = STableRow;
+	public:
+		
+		SLATE_BEGIN_ARGS(SCategoryColumnRow)
+		{}
+			SLATE_NAMED_SLOT(FArguments, Content)
+		SLATE_END_ARGS()
+
+		void Construct(const FArguments& InArgs, TSharedRef<STableViewBase> InOwner);
+		
+		virtual FReply OnMouseButtonDown(const FGeometry& MyGeometry, const FPointerEvent& MouseEvent) override;
+		virtual FReply OnMouseButtonDoubleClick(const FGeometry& InMyGeometry, const FPointerEvent& InMouseEvent) override;
+
+	private:
+
+		TSharedPtr<SWidget> Content;
+		
+		const FSlateBrush* GetBackgroundImage() const;
+	};
+}
diff --git a/Engine/Plugins/Developer/Concert/ConcertUI/ConcertSharedSlate/Source/ConcertSharedSlate/Private/Replication/Editor/View/Property/SMultiObjectAssignment.cpp b/Engine/Plugins/Developer/Concert/ConcertUI/ConcertSharedSlate/Source/ConcertSharedSlate/Private/Replication/Editor/View/Property/SMultiObjectAssignment.cpp
new file mode 100644
--- /dev/null
+++ b/Engine/Plugins/Developer/Concert/ConcertUI/ConcertSharedSlate/Source/ConcertSharedSlate/Private/Replication/Editor/View/Property/SMultiObjectAssignment.cpp
@@ -0,0 +1,192 @@
+﻿// Copyright Epic Games, Inc. All Rights Reserved.
+
+#include "SMultiObjectAssignment.h"
+
+#include "Misc/EBreakBehavior.h"
+#include "Replication/Data/ConcertPropertySelection.h"
+#include "Replication/Editor/Model/Object/IObjectHierarchyModel.h"
+#include "Replication/Editor/Utils/PropertyEnumerationUtils.h"
+
+namespace UE::ConcertSharedSlate
+{
+	namespace Private
+	{
+		static TMap<FSoftObjectPath, TArray<FSoftObjectPath>> BuildHierarchy(const TArray<FSoftObjectPath>& RootObjects, IObjectHierarchyModel& HierarchyModel, bool bShouldShowSubobjects)
+		{
+			TMap<FSoftObjectPath, TArray<FSoftObjectPath>> Result;
+
+			for (const FSoftObjectPath& Object : RootObjects)
+			{
+				HierarchyModel.ForEachChildRecursive(
+					Object,
+					[&Result](const FSoftObjectPath& Parent, const FSoftObjectPath& ChildObject, EChildRelationship Relationship)
+					{
+						// ForEachChildRecursive iterates through the hierarchy as it would be displayed by the component panel in the details panel
+						// E.g. for ACharacter this would say
+						//		Parent		= /Game/Maps.Map:PersistentLevel.Character.CapsuleComponent
+						//		ChildObject = /Game/Maps.Map:PersistentLevel.Character.CharacterMesh0
+						// However, we want Result to contain { /Game/Maps.Map:PersistentLevel.Character, /Game/Maps.Map:PersistentLevel.Character.CharacterMesh0 }
+						
+						// E.g. "/Game/Maps.Map:PersistentLevel.Actor.Component"
+						FString ParentPathString = ChildObject.ToString();
+						int32 Index = INDEX_NONE;
+						ParentPathString.FindLastChar('.', Index);
+						ParentPathString.LeftInline(Index);
+
+						const FSoftObjectPath ParentPath(*ParentPathString);
+						Result.FindOrAdd(ParentPath).Add(ChildObject);
+						return EBreakBehavior::Continue;
+					},
+					bShouldShowSubobjects ? EChildRelationshipFlags::All : EChildRelationshipFlags::Component
+					);
+			}
+
+			return Result;
+		}
+
+		static void ForEachCategory(
+			const TArray<FSoftObjectPath>& Start,
+			const TMap<FSoftObjectPath, TArray<FSoftObjectPath>>& Hierarchy,
+			TFunctionRef<void(const TArray<FSoftObjectPath>& ContextObjects)> Callback
+			)
+		{
+			if (Start.IsEmpty())
+			{
+				return;
+			}
+
+			// While multi-editing, if any of the subobjects do not match skip the entire category.
+			const FSoftObjectPath FirstEntry = *Start.CreateConstIterator();
+			const TArray<FSoftObjectPath>* FirstEntryChildren = Hierarchy.Find(FirstEntry);
+			if (!FirstEntryChildren)
+			{
+				return;
+			}
+
+			// E.g. "/Game/Maps.Map:PersistentLevel.Actor"
+			const FString Base = FirstEntry.ToString();
+			for (const FSoftObjectPath& ChildPath : *FirstEntryChildren)
+			{
+				// E.g. "/Game/Maps.Map:PersistentLevel.Actor.Component"
+				FString ChildName = ChildPath.ToString();
+				// +1 for the "." preceding the child, otherwise we'd get ".Component"
+				ChildName.RightChopInline(Base.Len() + 1);
+
+				TArray<FSoftObjectPath> RelatedChildren;
+				// This intentionally iterates FirstEntry as well
+				for (const FSoftObjectPath& StartObject : Start)
+				{
+					const TArray<FSoftObjectPath>* Children = Hierarchy.Find(StartObject);
+					if (!Children)
+					{
+						continue;
+					}
+
+					const FString ExpectedChildPathString = Base + TEXT(".") + ChildName;
+					const FSoftObjectPath ExpectedChildPath(*ExpectedChildPathString);
+					if (Children->Contains(ExpectedChildPath))
+					{
+						RelatedChildren.Add(ExpectedChildPath);
+					}
+				}
+
+				if (!RelatedChildren.IsEmpty())
+				{
+					Callback(RelatedChildren);
+					ForEachCategory(RelatedChildren, Hierarchy, Callback);
+				}
+			}
+		}
+	}
+	
+	void SMultiObjectAssignment::Construct(const FArguments& InArgs, TSharedRef<IPropertyTreeView> InTreeView)
+	{
+		TreeView = InTreeView;
+		ObjectHierarchy = InArgs._ObjectHierarchy;
+		OptionalPropertySource = InArgs._PropertySource;
+		
+		ChildSlot
+		[
+			TreeView->GetWidget()
+		];
+	}
+
+	void SMultiObjectAssignment::RefreshData(const TArray<FSoftObjectPath>& Objects, const IReplicationStreamModel& Model)
+	{
+		TArray<FPropertyAssignmentEntry> Entries;
+		auto[RootEntry, bRootSharesClass] = BuildAssignmentEntry(Objects, Model);
+		if (!ensureMsgf(bRootSharesClass, TEXT("Objects do not share the same class. Investigate invalid call.")))
+		{
+			PreviousSelectedObjects.Reset();
+			TreeView->RefreshPropertyData({}, false);
+		}
+		
+		Entries.Emplace(MoveTemp(RootEntry));
+		if (ObjectHierarchy)
+		{
+			using namespace Private;
+			const TMap<FSoftObjectPath, TArray<FSoftObjectPath>> HierarchyInfo = BuildHierarchy(Objects, *ObjectHierarchy, bShouldShowSubobjects);
+			ForEachCategory(Objects, HierarchyInfo, [this, &Model, &Entries](const TArray<FSoftObjectPath>& ContextObjects)
+			{
+				auto[Entry, bShareClass] = BuildAssignmentEntry(ContextObjects, Model);
+				if (bShareClass)
+				{
+					Entries.Emplace(MoveTemp(Entry));
+				}
+			});
+		}
+		
+		// If the objects have changed, the classes may share properties.
+		// In that case, below we'd reuse the item pointer, which would cause the tree view to re-use the old row widgets.
+		// However, we must regenerate all column widgets since they may be referencing the object the row was originally built for. So they'd display the state of the previous object still!
+		// Example: Assign property combo-box in Multi-User All Clients view displays who has the property assigned.
+		// Note: If the objects did not change, we definitely want to reuse item pointers since otherwise the user row selection is reset.
+		const bool bCanReusePropertyData = PreviousSelectedObjects == Objects;
+		PreviousSelectedObjects = Objects;
+			
+		TreeView->RefreshPropertyData(Entries, bCanReusePropertyData);
+	}
+
+	void SMultiObjectAssignment::SetShouldShowSubobjects(bool bShowSubobjects)
+	{
+		if (bShouldShowSubobjects != bShowSubobjects)
+		{
+			bShouldShowSubobjects = bShowSubobjects;
+			// TODO UE-216097: Refresh UI
+		}
+	}
+
+	SMultiObjectAssignment::FBuildAssignmentEntryResult SMultiObjectAssignment::BuildAssignmentEntry(const TArray<FSoftObjectPath>& Objects, const IReplicationStreamModel& Model)
+	{
+		FBuildAssignmentEntryResult Result;
+		
+		TArray<TSoftObjectPtr<>> ContextObjects;
+		Algo::Transform(Objects, ContextObjects, [](const FSoftObjectPath& Path){ return TSoftObjectPtr{ Path }; });
+
+		Result.Entry = { .ContextObjects = MoveTemp(ContextObjects) };
+		TSet<FConcertPropertyChain>& Properties = Result.Entry.PropertiesToDisplay;
+		FSoftClassPath& ClassPath = Result.Entry.Class;
+		
+		bool& bHasClassPath = Result.bHaveSharedClass;
+		EnumerateProperties(Objects, Model, OptionalPropertySource.Get(),
+			[&Properties, &bHasClassPath, &ClassPath](const FSoftClassPath& Class, const FConcertPropertyChain& Chain)
+			{
+				if (!bHasClassPath)
+				{
+					bHasClassPath = true;
+					ClassPath = Class;
+				}
+				else if (ClassPath != Class)
+				{
+					// Duplicate class - caller of RefreshData should not have called us but handle it gracefully.
+					bHasClassPath = false;
+					return EBreakBehavior::Break;
+				}
+				
+				Properties.Add(Chain);
+				return EBreakBehavior::Continue;
+			});
+
+		return Result;
+	}
+}
diff --git a/Engine/Plugins/Developer/Concert/ConcertUI/ConcertSharedSlate/Source/ConcertSharedSlate/Private/Replication/Editor/View/Property/SMultiObjectAssignment.h b/Engine/Plugins/Developer/Concert/ConcertUI/ConcertSharedSlate/Source/ConcertSharedSlate/Private/Replication/Editor/View/Property/SMultiObjectAssignment.h
new file mode 100644
--- /dev/null
+++ b/Engine/Plugins/Developer/Concert/ConcertUI/ConcertSharedSlate/Source/ConcertSharedSlate/Private/Replication/Editor/View/Property/SMultiObjectAssignment.h
@@ -0,0 +1,67 @@
+﻿// Copyright Epic Games, Inc. All Rights Reserved.
+
+#pragma once
+
+#include "Replication/Editor/View/IMultiObjectPropertyAssignmentView.h"
+#include "Replication/Editor/View/IPropertyAssignmentView.h"
+#include "Replication/Editor/View/IPropertyTreeView.h"
+
+#include "Widgets/DeclarativeSyntaxSupport.h"
+#include "Widgets/SCompoundWidget.h"
+
+namespace UE::ConcertSharedSlate
+{
+	class IObjectHierarchyModel;
+	class IPropertySelectionSourceModel;
+	class IReplicationStreamModel;
+	
+	/** SMultiObjectAssignment shows the properties of the displayed object and all of its subobjects.*/
+	class SMultiObjectAssignment
+		: public SCompoundWidget
+		, public IMultiObjectPropertyAssignmentView
+	{
+	public:
+
+		SLATE_BEGIN_ARGS(SMultiObjectAssignment){}
+			/** Optional. If specified, displays the properties of this model instead of those assigned in the stream. */
+			SLATE_ARGUMENT(TSharedPtr<IPropertySelectionSourceModel>, PropertySource)
+			
+			/** Optional. Gets components and subobjects of the displayed object. If unspecified, behaves exactly like SPerObjectAssignmentView. */
+			SLATE_ARGUMENT(TSharedPtr<IObjectHierarchyModel>, ObjectHierarchy)
+		SLATE_END_ARGS()
+
+		void Construct(const FArguments& InArgs, TSharedRef<IPropertyTreeView> InTreeView);
+
+		//~ Begin IPropertyAssignmentView Interface
+		virtual void RefreshData(const TArray<FSoftObjectPath>& Objects, const IReplicationStreamModel& Model) override;
+		virtual void RequestRefilter() const override { return TreeView->RequestRefilter(); }
+		virtual void RequestResortForColumn(const FName& ColumnId) override { return TreeView->RequestRefilter(); }
+		virtual TSharedRef<SWidget> GetWidget() override { return SharedThis(this); }
+		virtual void SetShouldShowSubobjects(bool bShowSubobjects) override;
+		virtual bool GetShouldShowSubobjects() const override { return bShouldShowSubobjects;}
+		//~ End IPropertyAssignmentView Interface
+
+	private:
+
+		/** The tree view that is being wrapped. */
+		TSharedPtr<IPropertyTreeView> TreeView;
+		/** Used to determine whether to rebuild the entire property data. */
+		TArray<FSoftObjectPath> PreviousSelectedObjects;
+
+		/** Used to get subobjects of selected objects. */
+		TSharedPtr<IObjectHierarchyModel> ObjectHierarchy;
+		/** Optional. If specified, displays the properties of this model instead of those assigned in the stream. */
+		TSharedPtr<IPropertySelectionSourceModel> OptionalPropertySource;
+
+		/** Whether EChildRelationshipFlags::Subobject objects should be shown. */
+		bool bShouldShowSubobjects = false;
+
+		struct FBuildAssignmentEntryResult
+		{
+			FPropertyAssignmentEntry Entry;
+			bool bHaveSharedClass = false;
+		};
+		/** Builds a property section grouped by Objects (usually has 1 object - contains similar objects, like StaticMeshComponent0, for multi-edit purposes). */
+		FBuildAssignmentEntryResult BuildAssignmentEntry(const TArray<FSoftObjectPath>& Objects, const IReplicationStreamModel& Model);
+	};
+}
diff --git a/Engine/Plugins/Developer/Concert/ConcertUI/ConcertSharedSlate/Source/ConcertSharedSlate/Private/Replication/Editor/View/Property/SPerObjectPropertyAssignment.cpp b/Engine/Plugins/Developer/Concert/ConcertUI/ConcertSharedSlate/Source/ConcertSharedSlate/Private/Replication/Editor/View/Property/SPerObjectPropertyAssignment.cpp
--- a/Engine/Plugins/Developer/Concert/ConcertUI/ConcertSharedSlate/Source/ConcertSharedSlate/Private/Replication/Editor/View/Property/SPerObjectPropertyAssignment.cpp
+++ b/Engine/Plugins/Developer/Concert/ConcertUI/ConcertSharedSlate/Source/ConcertSharedSlate/Private/Replication/Editor/View/Property/SPerObjectPropertyAssignment.cpp
@@ -19,9 +19,9 @@ namespace UE::ConcertSharedSlate
 		];
 	}
 
-	void SPerObjectPropertyAssignment::RefreshData(const TSet<FSoftObjectPath>& Objects, const IReplicationStreamModel& Model)
+	void SPerObjectPropertyAssignment::RefreshData(const TArray<FSoftObjectPath>& Objects, const IReplicationStreamModel& Model)
 	{
-		TSet<TSoftObjectPtr<>> ContextObjects;
+		TArray<TSoftObjectPtr<>> ContextObjects;
 		Algo::Transform(Objects, ContextObjects, [](const FSoftObjectPath& Path){ return TSoftObjectPtr{ Path }; });
 		
 		FPropertyAssignmentEntry AssignmentEntry { .ContextObjects = MoveTemp(ContextObjects) };
@@ -55,10 +55,10 @@ namespace UE::ConcertSharedSlate
 			// However, we must regenerate all column widgets since they may be referencing the object the row was originally built for. So they'd display the state of the previous object still!
 			// Example: Assign property combo-box in Multi-User All Clients view displays who has the property assigned.
 			// Note: If the objects did not change, we definitely want to reuse item pointers since otherwise the user row selection is reset.
-			const bool bCanReusePropertyData = PreviousSelectedObjects.Num() == Objects.Num() && PreviousSelectedObjects.Includes(Objects);
+			const bool bCanReusePropertyData = PreviousSelectedObjects == Objects;
 			PreviousSelectedObjects = Objects;
 			
-			TreeView->RefreshPropertyData({ AssignmentEntry } , bCanReusePropertyData);
+			TreeView->RefreshPropertyData({ AssignmentEntry }, bCanReusePropertyData);
 		}
 		else
 		{
diff --git a/Engine/Plugins/Developer/Concert/ConcertUI/ConcertSharedSlate/Source/ConcertSharedSlate/Private/Replication/Editor/View/Property/SPerObjectPropertyAssignment.h b/Engine/Plugins/Developer/Concert/ConcertUI/ConcertSharedSlate/Source/ConcertSharedSlate/Private/Replication/Editor/View/Property/SPerObjectPropertyAssignment.h
--- a/Engine/Plugins/Developer/Concert/ConcertUI/ConcertSharedSlate/Source/ConcertSharedSlate/Private/Replication/Editor/View/Property/SPerObjectPropertyAssignment.h
+++ b/Engine/Plugins/Developer/Concert/ConcertUI/ConcertSharedSlate/Source/ConcertSharedSlate/Private/Replication/Editor/View/Property/SPerObjectPropertyAssignment.h
@@ -28,7 +28,7 @@ namespace UE::ConcertSharedSlate
 		void Construct(const FArguments& InArgs, TSharedRef<IPropertyTreeView> InTreeView);
 
 		//~ Begin IPropertyAssignmentView Interface
-		virtual void RefreshData(const TSet<FSoftObjectPath>& Object, const IReplicationStreamModel& Model) override;
+		virtual void RefreshData(const TArray<FSoftObjectPath>& Objects, const IReplicationStreamModel& Model) override;
 		virtual void RequestRefilter() const override { return TreeView->RequestRefilter(); }
 		virtual void RequestResortForColumn(const FName& ColumnId) override { return TreeView->RequestRefilter(); }
 		virtual TSharedRef<SWidget> GetWidget() override { return SharedThis(this); }
@@ -39,7 +39,7 @@ namespace UE::ConcertSharedSlate
 		/** The tree view that is being wrapped. */
 		TSharedPtr<IPropertyTreeView> TreeView;
 		/** Used to determine whether to rebuild the entire property data. */
-		TSet<FSoftObjectPath> PreviousSelectedObjects;
+		TArray<FSoftObjectPath> PreviousSelectedObjects;
 		
 		/** Optional. If specified, displays the properties of this model instead of those assigned in the stream. */
 		TSharedPtr<IPropertySelectionSourceModel> OptionalPropertySource;
diff --git a/Engine/Plugins/Developer/Concert/ConcertUI/ConcertSharedSlate/Source/ConcertSharedSlate/Private/Replication/Editor/View/Property/SPropertyTreeView.cpp b/Engine/Plugins/Developer/Concert/ConcertUI/ConcertSharedSlate/Source/ConcertSharedSlate/Private/Replication/Editor/View/Property/SPropertyTreeView.cpp
--- a/Engine/Plugins/Developer/Concert/ConcertUI/ConcertSharedSlate/Source/ConcertSharedSlate/Private/Replication/Editor/View/Property/SPropertyTreeView.cpp
+++ b/Engine/Plugins/Developer/Concert/ConcertUI/ConcertSharedSlate/Source/ConcertSharedSlate/Private/Replication/Editor/View/Property/SPropertyTreeView.cpp
@@ -3,23 +3,52 @@
 #include "SPropertyTreeView.h"
 
 #include "Replication/Editor/Model/Data/PropertyData.h"
-
-#include "Algo/ForEach.h"
+#include "Replication/Editor/Model/Data/PropertyNodeData.h"
 #include "Replication/Editor/View/Column/PropertyColumnAdapter.h"
+#include "SCategoryColumnRow.h"
+#include "Trace/ConcertTrace.h"
+
+#include "Algo/Transform.h"
 #include "UObject/UnrealType.h"
 
 #define LOCTEXT_NAMESPACE "SReplicatedPropertiesView"
 
 namespace UE::ConcertSharedSlate
 {
+	namespace Private
+	{
+		static TSharedRef<FPropertyNodeData> AllocateNode_Property(TArray<TSoftObjectPtr<>> ContextObjects, FSoftClassPath OwningClass, FConcertPropertyChain PropertyChain)
+		{
+			return MakeShared<FPropertyNodeData>(FPropertyData{ MoveTemp(ContextObjects), MoveTemp(OwningClass), MoveTemp(PropertyChain) });
+		}
+
+		static TSharedRef<FPropertyNodeData> AllocateNode_Category(TArray<TSoftObjectPtr<>> ContextObjects)
+		{
+			return MakeShared<FPropertyNodeData>(FCategoryData{ ContextObjects });
+		}
+
+		static const TArray<TSoftObjectPtr<>>& GetContext(const FPropertyNodeData& ode)
+		{
+			return ode.IsCategoryNode()
+				? ode.GetCategoryData()->GetContextObjects()
+				: ode.GetPropertyData()->GetContextObjects();
+		}
+	}
+	
 	void SPropertyTreeView::Construct(const FArguments& InArgs)
 	{
+		FilterDelegate = InArgs._FilterItem;
+		CreateCategoryRowDelegate = InArgs._CreateCategoryRow;
+		
 		ChildSlot
 		[
-			SAssignNew(TreeView, SReplicationTreeView<FPropertyData>)
-				.RootItemsSource(&RootPropertyRowData)
+			SAssignNew(TreeView, SReplicationTreeView<FPropertyNodeData>)
+				.RootItemsSource(&RootPropertyNodes)
 				.OnGetChildren(this, &SPropertyTreeView::GetPropertyRowChildren)
-				.FilterItem(InArgs._FilterItem)
+				.FilterItem(this, &SPropertyTreeView::FilterItem)
+				.OverrideRowWidget(this, &SPropertyTreeView::OverrideRowWidget)
+				.OverrideIsLessThan(this, &SPropertyTreeView::OverrideIsLessThan)
+				.OverrideGetSearchTerms(this, &SPropertyTreeView::OverrideGetSearchTerms)
 				.Columns(FPropertyColumnAdapter::Transform(InArgs._Columns))
 				.ExpandableColumnLabel(InArgs._ExpandableColumnLabel)
 				.PrimarySort(InArgs._PrimarySort)
@@ -34,40 +63,49 @@ namespace UE::ConcertSharedSlate
 
 	void SPropertyTreeView::RefreshPropertyData(const TArray<FPropertyAssignmentEntry>& Entries, bool bCanReuseExistingRowItems)
 	{
+		SCOPED_CONCERT_TRACE(RefreshPropertyData);
 		if (!bCanReuseExistingRowItems)
 		{
 			ChainToPropertyDataCache.Reset();
+			CategoryNodes.Reset();
 		}
 		
-		// Try to re-use old instances by using the old ChainToPropertyDataCache. This is also done so the expansion states restore correctly in the tree view.
-		TMap<FConcertPropertyChain, TSharedPtr<FPropertyData>> NewChainToPropertyDataCache;
-		
 		PropertyRowData.Empty();
-		for (const FPropertyAssignmentEntry& Entry : Entries)
 		{
-			for (const FConcertPropertyChain& PropertyChain : Entry.PropertiesToDisplay)
+			// Try to re-use old instances. This is also done so the expansion states restore correctly in the tree view.
+			TMap<FConcertPropertyChain, TArray<TSharedPtr<FPropertyNodeData>>> NewChainToPropertyDataCache;
+			TMap<TSoftObjectPtr<>, FCategoryMetaData> NewCategoryNodes;
+
+			if (ShouldBuildCategories())
 			{
-				const TSharedPtr<FPropertyData>* ExistingItem = ChainToPropertyDataCache.Find(PropertyChain);
-				const TSharedRef<FPropertyData> Item = ExistingItem ? ExistingItem->ToSharedRef() : AllocatePropertyData(Entry.ContextObjects, Entry.Class, PropertyChain);
-				PropertyRowData.Emplace(Item);
-				NewChainToPropertyDataCache.Emplace(PropertyChain, Item);
+				RefreshCategoryNodes(Entries, NewCategoryNodes);
 			}
-		}
+			RefreshPropertyNodes(Entries, NewChainToPropertyDataCache);
 
-		// If an item was removed, then NewPathToPropertyDataCache does not contain it. 
-		ChainToPropertyDataCache = MoveTemp(NewChainToPropertyDataCache);
+			// If an item was removed, then NewPathToPropertyDataCache does not contain it. 
+			ChainToPropertyDataCache = MoveTemp(NewChainToPropertyDataCache);
+			CategoryNodes = MoveTemp(NewCategoryNodes);
+		}
 		
 		// The tree view requires the item source to only contain the root items.
 		BuildRootPropertyRowData();
-
 		TreeView->RequestRefilter();
+
+		// Every time the tree is filled with a new hierarchy, all categories should be expanded
+		if (!bCanReuseExistingRowItems && ShouldBuildCategories())
+		{
+			TArray<TSharedPtr<FPropertyNodeData>> Categories;
+			Algo::Transform(CategoryNodes, Categories, [](const TPair<TSoftObjectPtr<>, FCategoryMetaData>& Pair){ return Pair.Value.Node.ToSharedPtr(); });
+			TreeView->SetExpandedItems(Categories, true);
+		}
 	}
 
 	void SPropertyTreeView::RequestScrollIntoView(const FConcertPropertyChain& PropertyChain)
 	{
-		const int32 Index = PropertyRowData.IndexOfByPredicate([&PropertyChain](const TSharedPtr<FPropertyData>& Data)
+		const int32 Index = PropertyRowData.IndexOfByPredicate([&PropertyChain](const TSharedPtr<FPropertyNodeData>& Data)
 		{
-			return Data->GetProperty() == PropertyChain;
+			const TOptional<FPropertyData>& PropertyData = Data->GetPropertyData();
+			return PropertyData && PropertyData->GetProperty() == PropertyChain;
 		});
 		if (PropertyRowData.IsValidIndex(Index))
 		{
@@ -75,38 +113,220 @@ namespace UE::ConcertSharedSlate
 			TreeView->RequestScrollIntoView(PropertyRowData[Index]);
 		}
 	}
+	
+	void SPropertyTreeView::RefreshCategoryNodes(const TArray<FPropertyAssignmentEntry>& Entries, TMap<TSoftObjectPtr<>, FCategoryMetaData>& NewCategoryNodes)
+	{
+		SourceEntriesForSorting.Empty();
+				
+		for (const FPropertyAssignmentEntry& Entry : Entries)
+		{
+			if (!ensure(!Entry.ContextObjects.IsEmpty()))
+			{
+				continue;
+			}
+			// Entry's category widget is supposed to appear in the same order as it was passed to us
+			SourceEntriesForSorting.Add(Entry.ContextObjects);
+			const TSoftObjectPtr<> FirstContextObject = Entry.ContextObjects[0];
+
+			// Re-use existing FPropertyNodeData if possible (to retain selection in tree view)
+			const FCategoryMetaData* ExistingItem = CategoryNodes.Find(FirstContextObject);
+			if (!ExistingItem)
+			{
+				const TSharedRef<FPropertyNodeData> NodeData = Private::AllocateNode_Category(Entry.ContextObjects);
+				const FCategoryRowGenerationArgs CategoryArgs(Entry.ContextObjects, TAttribute<FText>::CreateLambda([this](){ return TreeView->GetHighlightText(); }));
+				const TSharedRef<ICategoryRow> CategoryRow = CreateCategoryRowDelegate.Execute(CategoryArgs);
+				
+				FCategoryMetaData& CategoryMetaData = CategoryNodes.Emplace(FirstContextObject, FCategoryMetaData(NodeData, CategoryRow));
+				ExistingItem = &CategoryMetaData;
+			}
+			
+			PropertyRowData.Emplace(ExistingItem->Node);
+
+			// Even when multi-editing, only the first context object is mapped to the category.
+			NewCategoryNodes.Emplace(FirstContextObject, *ExistingItem);
+		}
+	}
 
-	inline TSharedRef<FPropertyData> SPropertyTreeView::AllocatePropertyData(TSet<TSoftObjectPtr<>> ContextObjects, FSoftClassPath OwningClass, FConcertPropertyChain PropertyChain)
+	void SPropertyTreeView::RefreshPropertyNodes(const TArray<FPropertyAssignmentEntry>& Entries, TMap<FConcertPropertyChain, TArray<TSharedPtr<FPropertyNodeData>>>& NewChainToPropertyDataCache)
 	{
-		return MakeShared<FPropertyData>(MoveTemp(ContextObjects), MoveTemp(OwningClass), MoveTemp(PropertyChain));
+		for (const FPropertyAssignmentEntry& Entry : Entries)
+		{
+#if UE_BUILD_DEBUG
+			TSet<FConcertPropertyChain> DuplicateChainDetection;
+#endif
+			
+			for (const FConcertPropertyChain& PropertyChain : Entry.PropertiesToDisplay)
+			{
+#if UE_BUILD_DEBUG
+				check(!DuplicateChainDetection.Contains(PropertyChain));
+				DuplicateChainDetection.Add(PropertyChain);
+#endif
+				
+				TSharedPtr<FPropertyNodeData> PreExisting = FindPropertyNode(PropertyChain, Entry.ContextObjects);
+				const TSharedRef<FPropertyNodeData> Item = PreExisting
+					? PreExisting.ToSharedRef()
+					: Private::AllocateNode_Property(Entry.ContextObjects, Entry.Class, PropertyChain);
+				PropertyRowData.Emplace(Item);
+				NewChainToPropertyDataCache.FindOrAdd(PropertyChain).AddUnique(Item);
+			}
+		}
 	}
 
-	void SPropertyTreeView::BuildRootPropertyRowData()
+	TSharedPtr<FPropertyNodeData> SPropertyTreeView::FindPropertyNode(const FConcertPropertyChain& PropertyChain, TConstArrayView<TSoftObjectPtr<>> ContextObjects)
 	{
-		RootPropertyRowData.Empty(PropertyRowData.Num());
-		for (const TSharedPtr<FPropertyData>& PropertyData : PropertyRowData)
+		const TArray<TSharedPtr<FPropertyNodeData>>* ExistingItems = ChainToPropertyDataCache.Find(PropertyChain);
+		if (!ExistingItems)
+		{
+			return nullptr;
+		}
+
+		for (const TSharedPtr<FPropertyNodeData>& Item : *ExistingItems)
 		{
-			if (PropertyData->GetProperty().IsRootProperty())
+			const TOptional<FPropertyData>& PropertyData = Item->GetPropertyData();
+			const bool bIsProperty = PropertyData->GetProperty() == PropertyChain;
+			const bool bReferencesContextObject = PropertyData->GetContextObjects().ContainsByPredicate([&ContextObjects](const TSoftObjectPtr<>& ContextObject)
 			{
-				RootPropertyRowData.Emplace(PropertyData);
+				return ContextObjects.Contains(ContextObject);
+			});
+			if (bIsProperty && bReferencesContextObject)
+			{
+				return Item;
 			}
 		}
+
+		return nullptr;
 	}
 
-	void SPropertyTreeView::GetPropertyRowChildren(TSharedPtr<FPropertyData> ReplicatedPropertyData, TFunctionRef<void(TSharedPtr<FPropertyData>)> ProcessChild)
+	void SPropertyTreeView::BuildRootPropertyRowData()
+	{
+		RootPropertyNodes.Empty(PropertyRowData.Num());
+
+		if (ShouldBuildCategories())
+		{
+			Algo::TransformIf(PropertyRowData, RootPropertyNodes,
+				[](const TSharedPtr<FPropertyNodeData>& NodeData){ return NodeData->IsCategoryNode(); },
+				[](const TSharedPtr<FPropertyNodeData>& NodeData){ return NodeData; }
+				);
+		}
+		else
+		{
+			Algo::TransformIf(PropertyRowData, RootPropertyNodes,
+				[](const TSharedPtr<FPropertyNodeData>& NodeData)
+				{
+					const TOptional<FPropertyData>& PropertyData = NodeData->GetPropertyData();
+					return PropertyData && PropertyData->GetProperty().IsRootProperty();
+				},
+				[](const TSharedPtr<FPropertyNodeData>& NodeData){ return NodeData; }
+				);
+		}
+	}
+	
+	void SPropertyTreeView::GetPropertyRowChildren(TSharedPtr<FPropertyNodeData> ParentNodeData, TFunctionRef<void(TSharedPtr<FPropertyNodeData>)> ProcessChild) const
+	{
+		SCOPED_CONCERT_TRACE(GetPropertyRowChildren);
+		
+		if (ParentNodeData->IsCategoryNode())
+		{
+			const TOptional<FCategoryData>& ParentCategoryData = ParentNodeData->GetCategoryData();
+			check(ParentCategoryData);
+			EnumerateRootProperties(*ParentCategoryData, ProcessChild);
+		}
+		else
+		{
+			EnumerateChildProperties(*ParentNodeData, ProcessChild);
+		}
+	}
+	
+	void SPropertyTreeView::EnumerateRootProperties(const FCategoryData& ParentCategoryData, TFunctionRef<void(TSharedPtr<FPropertyNodeData>)> ProcessChild) const
+	{
+		for (const TSharedPtr<FPropertyNodeData>& NodeData : PropertyRowData)
+		{
+			const TOptional<FPropertyData>& PropertyData = NodeData->GetPropertyData();
+			if (!PropertyData)
+			{
+				continue;
+			}
+
+			const TArray<TSoftObjectPtr<>>& ParentContextObjects = ParentCategoryData.GetContextObjects();
+			const TArray<TSoftObjectPtr<>>& ContextObjects = PropertyData->GetContextObjects();
+			const bool bHaveSameContextObject = ContextObjects == ParentContextObjects;
+			if (bHaveSameContextObject && PropertyData->GetProperty().IsRootProperty())
+			{
+				ProcessChild(NodeData);
+			}
+		}
+	}
+	
+	void SPropertyTreeView::EnumerateChildProperties(const FPropertyNodeData& ParentNode, TFunctionRef<void(TSharedPtr<FPropertyNodeData>)> ProcessChild) const
 	{
-		TArray<TSharedPtr<FPropertyData>> Children;
+		const TArray<TSoftObjectPtr<>>& ParentContext = Private::GetContext(ParentNode);
+		const TOptional<FPropertyData>& ParentPropertyData = ParentNode.GetPropertyData();
+		check(ParentPropertyData);
 		
 		// Not the most efficient but it should be fine.
-		for (const TSharedPtr<FPropertyData>& Data : PropertyRowData)
+		for (const TSharedPtr<FPropertyNodeData>& NodeData : PropertyRowData)
 		{
-			if (Data->GetProperty().IsDirectChildOf(ReplicatedPropertyData->GetProperty()))
+			// There can be multiple components with the same property paths - exclude properties that are under different category nodes.
+			const TArray<TSoftObjectPtr<>>& OtherNodeContext = Private::GetContext(*NodeData);
+			const bool bContextIsSame = ParentContext == OtherNodeContext;
+			if (!bContextIsSame)
 			{
-				Children.Add(Data);
+				continue;
+			}
+			
+			const TOptional<FPropertyData>& PropertyData = NodeData->GetPropertyData();
+			if (bContextIsSame && PropertyData && PropertyData->GetProperty().IsDirectChildOf(ParentPropertyData->GetProperty()))
+			{
+				ProcessChild(NodeData);
 			}
 		}
+	}
 
-		Algo::ForEach(Children, [&ProcessChild](const TSharedPtr<FPropertyData>& Data){ ProcessChild(Data); });
+	EItemFilterResult SPropertyTreeView::FilterItem(const FPropertyNodeData& PropertyNodeData) const
+	{
+		return PropertyNodeData.IsCategoryNode()
+			? EItemFilterResult::IncludeOnlyIfChildIsIncluded
+			: FilterDelegate.IsBound() ? FilterDelegate.Execute(*PropertyNodeData.GetPropertyData()) : EItemFilterResult::Include;
+	}
+
+	TSharedPtr<ITableRow> SPropertyTreeView::OverrideRowWidget(TSharedPtr<FPropertyNodeData> NodeData, const TSharedRef<STableViewBase>& TableViewBase, const TReplicationTreeData<FPropertyNodeData>::FGenerateRowArgs& Args)
+	{
+		if (NodeData->IsCategoryNode() && ensure(CreateCategoryRowDelegate.IsBound()))
+		{
+			const FCategoryMetaData* CategoryMetaData = FindCategoryMetaData(*NodeData);
+			return SNew(SCategoryColumnRow, TableViewBase)
+				.Content()
+				[
+					CategoryMetaData->RowWidgetContent->GetWidget()
+				];
+		}
+		return nullptr;
+	}
+
+	EComparisonOverride SPropertyTreeView::OverrideIsLessThan(const TSharedPtr<FPropertyNodeData>& Left, const TSharedPtr<FPropertyNodeData>& Right) const
+	{
+		if (!Left->IsCategoryNode() && !Right->IsCategoryNode())
+		{
+			return EComparisonOverride::UseDefault;
+		}
+		
+		const TArray<TSoftObjectPtr<>>& LeftContext = Private::GetContext(*Left);
+		const TArray<TSoftObjectPtr<>>& RightContext = Private::GetContext(*Right);
+		const int32 LeftIndex = SourceEntriesForSorting.IndexOfByKey(LeftContext);
+		const int32 RightIndex = SourceEntriesForSorting.IndexOfByKey(RightContext);
+		ensure(SourceEntriesForSorting.IsValidIndex(LeftIndex) && SourceEntriesForSorting.IsValidIndex(RightIndex));
+		return LeftIndex < RightIndex ? EComparisonOverride::Less : EComparisonOverride::NotLess;
+	}
+
+	ESearchTermResult SPropertyTreeView::OverrideGetSearchTerms(const TSharedPtr<FPropertyNodeData>& NodeData, TArray<FString>& InOutSearchStrings) const
+	{
+		const FCategoryMetaData* CategoryMetaData = NodeData->IsCategoryNode() ? FindCategoryMetaData(*NodeData) : nullptr;
+		if (CategoryMetaData)
+		{
+			CategoryMetaData->RowWidgetContent->GenerateSearchTerms(NodeData->GetCategoryData()->GetContextObjects(), InOutSearchStrings);
+			return ESearchTermResult::UseOverrideOnly;
+		}
+		return ESearchTermResult::UseDefault;
 	}
 }
 
diff --git a/Engine/Plugins/Developer/Concert/ConcertUI/ConcertSharedSlate/Source/ConcertSharedSlate/Private/Replication/Editor/View/Property/SPropertyTreeView.h b/Engine/Plugins/Developer/Concert/ConcertUI/ConcertSharedSlate/Source/ConcertSharedSlate/Private/Replication/Editor/View/Property/SPropertyTreeView.h
--- a/Engine/Plugins/Developer/Concert/ConcertUI/ConcertSharedSlate/Source/ConcertSharedSlate/Private/Replication/Editor/View/Property/SPropertyTreeView.h
+++ b/Engine/Plugins/Developer/Concert/ConcertUI/ConcertSharedSlate/Source/ConcertSharedSlate/Private/Replication/Editor/View/Property/SPropertyTreeView.h
@@ -2,7 +2,9 @@
 
 #pragma once
 
+#include "Replication/PropertyTreeFactory.h"
 #include "Replication/Data/ConcertPropertySelection.h"
+#include "Replication/Editor/Model/Data/PropertyNodeData.h"
 #include "Replication/Editor/View/IPropertyTreeView.h"
 #include "Replication/Editor/View/Tree/SReplicationTreeView.h"
 
@@ -12,7 +14,13 @@
 
 namespace UE::ConcertSharedSlate
 {
-	class FPropertyData;
+	class ICategoryRow;
+}
+
+namespace UE::ConcertSharedSlate
+{
+	class FCategoryData;
+	class FPropertyNodeData;
 	
 	/**
 	 * This widget knows how to display a list of properties in a tree view.
@@ -28,8 +36,11 @@ namespace UE::ConcertSharedSlate
 		{}
 			/*************** Arguments inherited by SReplicationTreeView ***************/
 		
-			/** Optional callback to do even more filtering of items. */
+			/** Optional callback to do more filtering of items on top of the search bar. */
 			SLATE_EVENT(SReplicationTreeView<FPropertyData>::FCustomFilter, FilterItem)
+			
+			/** Optional callback for creating category rows. If unset, no category rows are generated. */
+			SLATE_EVENT(FCreateCategoryRow, CreateCategoryRow)
 		
 			/** The columns this list should have */
 			SLATE_ARGUMENT(TArray<FPropertyColumnEntry>, Columns)
@@ -69,20 +80,82 @@ namespace UE::ConcertSharedSlate
 	private:
 
 		/** The tree view displaying the replicated properties */
-		TSharedPtr<SReplicationTreeView<FPropertyData>> TreeView;
+		TSharedPtr<SReplicationTreeView<FPropertyNodeData>> TreeView;
 		
 		/** Contains all data. */
-		TArray<TSharedPtr<FPropertyData>> PropertyRowData;
+		TArray<TSharedPtr<FPropertyNodeData>> PropertyRowData;
 		/** The instances which do not have any parents. This acts as the item source for the tree view. */
-		TArray<TSharedPtr<FPropertyData>> RootPropertyRowData;
-		/** Inverse map of PropertyRowData using FPropertyData::GetProperty as key. Contains all elements of PropertyRowData. */
-		TMap<FConcertPropertyChain, TSharedPtr<FPropertyData>> ChainToPropertyDataCache;
+		TArray<TSharedPtr<FPropertyNodeData>> RootPropertyNodes;
+		
+		struct FCategoryMetaData
+		{
+			TSharedRef<FPropertyNodeData> Node;
+			/** The widget displayed in the category widget. */
+			TSharedRef<ICategoryRow> RowWidgetContent;
+
+			FCategoryMetaData(TSharedRef<FPropertyNodeData> Node, TSharedRef<ICategoryRow> RowWidgetContent)
+				: Node(Node)
+				, RowWidgetContent(RowWidgetContent) 
+			{}
+		};
+		/**
+		 * Inverse map of object to its owning category.
+		 * Contains all elements of PropertyRowData which are category nodes.
+		 *
+		 * The main purpose is to allow reuse of items when refreshing hierarchy (to retain selection state, etc.)
+		 * 
+		 * Important: when multi-editing, this only maps the context object at index 0.
+		 * Use FindCategoryMetaData to retrieve.
+		 * 
+		 * Empty if ShouldBuildCategories() == false.
+		 */
+		TMap<TSoftObjectPtr<>, FCategoryMetaData> CategoryNodes;
+		/**
+		 * Inverse map of PropertyRowData using FPropertyData::GetProperty as key.
+		 * Contains all elements of PropertyRowData which are property nodes
+		 *
+		 * The main purpose is to allow reuse of items when refreshing hierarchy (to retain selection state, etc.)
+		 */
+		TMap<FConcertPropertyChain, TArray<TSharedPtr<FPropertyNodeData>>> ChainToPropertyDataCache;
 		
-		TSharedRef<FPropertyData> AllocatePropertyData(TSet<TSoftObjectPtr<>> ContextObjects, FSoftClassPath OwningClass, FConcertPropertyChain PropertyChain);
+		/** Equal to FPropertyAssignmentEntry::ContextObjects from the last RefreshPropertyData call. This is used during sorting to determine the relative ordering of categories. */
+		TArray<TArray<TSoftObjectPtr<>>> SourceEntriesForSorting;
+
+		/** Optional callback for filtering items. */
+		SReplicationTreeView<FPropertyData>::FCustomFilter FilterDelegate;
+		/** Optional callback for creating category rows. If unset, no category rows are generated. */
+		FCreateCategoryRow CreateCategoryRowDelegate;
+
+		/** @return Whether this UI was configured to show category nodes. */
+		bool ShouldBuildCategories() const { return CreateCategoryRowDelegate.IsBound(); }
+		
+		void RefreshCategoryNodes(const TArray<FPropertyAssignmentEntry>& Entries, TMap<TSoftObjectPtr<>, FCategoryMetaData>& NewCategoryNodes);
+		void RefreshPropertyNodes(const TArray<FPropertyAssignmentEntry>& Entries, TMap<FConcertPropertyChain, TArray<TSharedPtr<FPropertyNodeData>>>& NewChainToPropertyDataCache);
+
+		/** Looks in ChainToPropertyDataCache for any a node that is referencing PropertyChain and any of ContextObjects. */
+		TSharedPtr<FPropertyNodeData> FindPropertyNode(const FConcertPropertyChain& PropertyChain, TConstArrayView<TSoftObjectPtr<>> ContextObjects);
+		/** Find the meta data for this item. */
+		FCategoryMetaData* FindCategoryMetaData(const FPropertyNodeData& Item) { return CategoryNodes.Find(Item.GetCategoryData()->GetContextObjects()[0]); }
+		const FCategoryMetaData* FindCategoryMetaData(const FPropertyNodeData& Item) const { check(Item.IsCategoryNode()); return CategoryNodes.Find(Item.GetCategoryData()->GetContextObjects()[0]); }
 
 		/** Inits RootPropertyRowData from PropertyRowData. */
 		void BuildRootPropertyRowData();
-		void GetPropertyRowChildren(TSharedPtr<FPropertyData> ReplicatedPropertyData, TFunctionRef<void(TSharedPtr<FPropertyData>)> ProcessChild);
+		void GetPropertyRowChildren(TSharedPtr<FPropertyNodeData> ParentNodeData, TFunctionRef<void(TSharedPtr<FPropertyNodeData>)> ProcessChild) const;
+		void EnumerateRootProperties(const FCategoryData& ParentCategoryData, TFunctionRef<void(TSharedPtr<FPropertyNodeData>)> ProcessChild) const;
+		void EnumerateChildProperties(const FPropertyNodeData& ParentNode, TFunctionRef<void(TSharedPtr<FPropertyNodeData>)> ProcessChild) const;
+
+		/** Filters the item according to FilterDelegate. */
+		EItemFilterResult FilterItem(const FPropertyNodeData& PropertyNodeData) const;
+		/** Generates a category row if NodeData is a category. */
+		TSharedPtr<ITableRow> OverrideRowWidget(
+			TSharedPtr<FPropertyNodeData> NodeData,
+			const TSharedRef<STableViewBase>& TableViewBase,
+			const TReplicationTreeData<FPropertyNodeData>::FGenerateRowArgs& Args
+			);
+		/** Overrides the sort if one of the nodes is a category */
+		EComparisonOverride OverrideIsLessThan(const TSharedPtr<FPropertyNodeData>& Left, const TSharedPtr<FPropertyNodeData>& Right) const;
+		/** Overrides the search terms for category nodes. */
+		ESearchTermResult OverrideGetSearchTerms(const TSharedPtr<FPropertyNodeData>& NodeData, TArray<FString>& InOutSearchStrings) const;
 	};
 }
 
diff --git a/Engine/Plugins/Developer/Concert/ConcertUI/ConcertSharedSlate/Source/ConcertSharedSlate/Private/Replication/Editor/View/Tree/SReplicationTreeView.h b/Engine/Plugins/Developer/Concert/ConcertUI/ConcertSharedSlate/Source/ConcertSharedSlate/Private/Replication/Editor/View/Tree/SReplicationTreeView.h
--- a/Engine/Plugins/Developer/Concert/ConcertUI/ConcertSharedSlate/Source/ConcertSharedSlate/Private/Replication/Editor/View/Tree/SReplicationTreeView.h
+++ b/Engine/Plugins/Developer/Concert/ConcertUI/ConcertSharedSlate/Source/ConcertSharedSlate/Private/Replication/Editor/View/Tree/SReplicationTreeView.h
@@ -32,6 +32,26 @@ namespace UE::ConcertSharedSlate
 		/** The item is only displayed if the filter passes on one of its children. If no child is included, the item is not displayed. */
 		IncludeOnlyIfChildIsIncluded
 	};
+
+	/** Result for the IsLessThan override delegate in SReplicationTreeView. */
+	enum class EComparisonOverride : uint8
+	{
+		/** Left < Right is to be treated as true */
+		Less,
+		/** Left < Right is to be treated as false */
+		NotLess,
+		/** Use the standard column comparision delegates to find out. */
+		UseDefault
+	};
+
+	/** Result for the GetSearchString override delegate in SReplicationTreeView. */
+	enum class ESearchTermResult : uint8
+	{
+		/** The delegate wants the passed in item to go through normal text search as well. */
+		UseDefault,
+		/** The delegate does not want the passed in item to go through the rest of the text search process. */
+		UseOverrideOnly,
+	};
 	
 	/**
 	 * Shared code for the tree view for replicated actors and properties.
@@ -44,16 +64,18 @@ namespace UE::ConcertSharedSlate
 	class SReplicationTreeView : public SCompoundWidget
 	{
 	public:
-
-		using TOverrideColumnWidget = typename TReplicationTreeData<TItemType>::FOverrideColumnWidget;
-		using FGetHoveredRowContent = typename TReplicationTreeData<TItemType>::FGetHoveredRowContent;
-
+		
+		using FOverrideColumnWidget = typename TReplicationTreeData<TItemType>::FOverrideColumnWidget;
 		DECLARE_DELEGATE_OneParam(FDeleteItems, const TArray<TSharedPtr<TItemType>>& SelectedItems);
-		DECLARE_DELEGATE_TwoParams(FGetItemChildren, TSharedPtr<TItemType> Item, TFunctionRef<void(TSharedPtr<TItemType>)> ProcessChild);
 		DECLARE_DELEGATE(FOnSelectionChanged);
 
+		// TODO UE-216456: The following callbacks could be extracted to a IReplicationItem<TItemType> to simplify SReplicationTreeView implementation 
+		using FOverrideRowWidget	= typename TReplicationTreeData<TItemType>::FOverrideRowWidget;
+		using FGetHoveredRowContent = typename TReplicationTreeData<TItemType>::FGetHoveredRowContent;
 		DECLARE_DELEGATE_RetVal_OneParam(EItemFilterResult, FCustomFilter, const TItemType& Item);
-		DECLARE_DELEGATE_RetVal_OneParam(bool, FIsSearchableItem, const TSharedPtr<TItemType>& Item);
+		DECLARE_DELEGATE_RetVal_TwoParams(EComparisonOverride, FIsLessThanOverride, const TSharedPtr<TItemType>& Left, const TSharedPtr<TItemType>& Right);
+		DECLARE_DELEGATE_TwoParams(FGetItemChildren, TSharedPtr<TItemType> Item, TFunctionRef<void(TSharedPtr<TItemType>)> ProcessChild);
+		DECLARE_DELEGATE_RetVal_TwoParams(ESearchTermResult, FGetSearchTermsOverride, const TSharedPtr<TItemType>& Item, TArray<FString>& InOutSearchTerms);
 
 		enum class EContent
 		{
@@ -83,17 +105,19 @@ namespace UE::ConcertSharedSlate
 
 			/** Optional callback to do even more filtering of items. */
 			SLATE_EVENT(FCustomFilter, FilterItem)
-		
+			/** Optional. Gets the content to overlay on hovered rows; it covers the entire row. */
+			SLATE_EVENT(FGetHoveredRowContent, GetHoveredRowContent)
 			/**
 			 * Optional. If the delegate returns non-null, that widget will be used instead of the one the column would generate.
 			 * This is useful, e.g. if you want to generate a separator widget between items.
 			 */
-			SLATE_EVENT(TOverrideColumnWidget, OverrideColumnWidget)
-			/** Optional callback for determining whether this item can be searched. */
-			SLATE_EVENT(FIsSearchableItem, IsSearchableItem)
-		
-			/** Optional. Gets the content to overlay on hovered rows; it covers the entire row. */
-			SLATE_EVENT(FGetHoveredRowContent, GetHoveredRowContent)
+			SLATE_EVENT(FOverrideColumnWidget, OverrideColumnWidget)
+			/** Optional. The widget returned by this delegate ends up overriding the default row widget. */
+			SLATE_EVENT(FOverrideRowWidget, OverrideRowWidget)
+			/** Optional. This delegate is used during sorting and can override the default sorting behaviour for some items. This can be used e.g. for category nodes. */
+			SLATE_EVENT(FIsLessThanOverride, OverrideIsLessThan)
+			/** Optional. Can generate custom terms for an item type. This can be used e.g. for category nodes. */
+			SLATE_EVENT(FGetSearchTermsOverride, OverrideGetSearchTerms)
 			
 			/** The columns this list should have */
 			SLATE_ARGUMENT(TArray<TReplicationColumnEntry<TItemType>>, Columns)
@@ -131,14 +155,18 @@ namespace UE::ConcertSharedSlate
 			OnGetChildrenDelegate = InArgs._OnGetChildren;
 			OnDeleteItemsDelegate = InArgs._OnDeleteItems;
 			CustomFilterDelegate = InArgs._FilterItem;
-			OverrideColumnWidget = InArgs._OverrideColumnWidget;
-			IsSearchableItemDelegate = InArgs._IsSearchableItem;
 			GetHoveredRowContentDelegate = InArgs._GetHoveredRowContent;
+			OverrideColumnWidgetDelegate = InArgs._OverrideColumnWidget;
+			OverrideRowWidgetDelegate = InArgs._OverrideRowWidget;
+			OverrideIsLessThanDelegate = InArgs._OverrideIsLessThan;
+			OverrideGetSearchTerms = InArgs._OverrideGetSearchTerms;
 			ExpandableColumnId = InArgs._ExpandableColumnLabel;
 			RowStyle = InArgs._RowStyle;
 			
 			SearchText = MakeShared<FText>();
-			SearchTextFilter = MakeShared<TTextFilter<const TSharedPtr<TItemType>&>>(TTextFilter<const TSharedPtr<TItemType>&>::FItemToStringArray::CreateSP(this, &SReplicationTreeView::PopulateSearchStrings));
+			SearchTextFilter = MakeShared<TTextFilter<const TSharedPtr<TItemType>&>>(
+				TTextFilter<const TSharedPtr<TItemType>&>::FItemToStringArray::CreateSP(this, &SReplicationTreeView::PopulateSearchStrings)
+				);
 			SearchTextFilter->OnChanged().AddSP(this, &SReplicationTreeView::RequestRefilter);
 			
 			ChildSlot
@@ -291,6 +319,7 @@ namespace UE::ConcertSharedSlate
 		
 		TArray<TSharedPtr<TItemType>> GetSelectedItems() const { return TreeView->GetSelectedItems(); }
 		const TArray<TSharedPtr<TItemType>>& GetFilteredRootItems() const { return FilteredRootItems; }
+		const FText& GetHighlightText() const { return *SearchText; }
 		
 		virtual FReply OnKeyDown(const FGeometry& MyGeometry, const FKeyEvent& InKeyEvent) override;
 		virtual void Tick(const FGeometry& AllottedGeometry, const double InCurrentTime, const float InDeltaTime) override;
@@ -351,12 +380,16 @@ namespace UE::ConcertSharedSlate
 		FDeleteItems OnDeleteItemsDelegate;
 		/** Optional delegate for filtering the items even more. */
 		FCustomFilter CustomFilterDelegate;
-		/** Optional delegate for overriding the column widgets. */
-		TOverrideColumnWidget OverrideColumnWidget;
-		/** Optional callback for determining whether this item can be filtered. If false, it will not be shown when searched. */
-		FIsSearchableItem IsSearchableItemDelegate;
 		/** Optional. The content to overlay on hovered rows; it covers the entire row. */
 		FGetHoveredRowContent GetHoveredRowContentDelegate;
+		/** Optional delegate for overriding the column widgets. */
+		FOverrideColumnWidget OverrideColumnWidgetDelegate;
+		/** Optional. The widget returned by this delegate ends up overriding the default row widget. */
+		FOverrideRowWidget OverrideRowWidgetDelegate;
+		/** Optional. This delegate is used during sorting and can override the default sorting behaviour for some items. */
+		FIsLessThanOverride OverrideIsLessThanDelegate;
+		/** Optional. Can generate custom terms for an item type. This can be used e.g. for category nodes. */
+		FGetSearchTermsOverride OverrideGetSearchTerms;
 
 		/** Style to use for rows */
 		const FTableRowStyle* RowStyle = nullptr;
@@ -516,16 +549,23 @@ namespace UE::ConcertSharedSlate
 			SReplicationColumnRow<TItemType>::FGetColumn::CreateSP(this, &SReplicationTreeView::FindColumnByName
 			);
 
-		// Compile-time decide the row that is supposed to be generated
-		return TReplicationTreeItemTraits<TItemType>::GenerateRowWidget(Item, OwnerTable,
-			{
-				ColumnGetter,
-				OverrideColumnWidget,
-				GetHoveredRowContentDelegate,
-				SearchText,
-				ExpandableColumnId,
-				RowStyle
-			});
+		const typename TReplicationTreeData<TItemType>::FGenerateRowArgs RowArgs
+		{
+			ColumnGetter,
+			OverrideColumnWidgetDelegate,
+			GetHoveredRowContentDelegate,
+			SearchText,
+			ExpandableColumnId,
+			RowStyle
+		};
+		const TSharedPtr<ITableRow> OverrideRowWidget = OverrideRowWidgetDelegate.IsBound()
+			? OverrideRowWidgetDelegate.Execute(Item, OwnerTable, RowArgs)
+			: nullptr;
+		
+		return OverrideRowWidget
+			? OverrideRowWidget.ToSharedRef()
+			// Compile-time decide the row that is supposed to be generated
+			: TReplicationTreeItemTraits<TItemType>::GenerateRowWidget(Item, OwnerTable, RowArgs);
 	}
 
 	template <typename TItemType>
@@ -626,11 +666,12 @@ namespace UE::ConcertSharedSlate
 	template <typename TItemType>
 	void SReplicationTreeView<TItemType>::PopulateSearchStrings(const TSharedPtr<TItemType>& Item, TArray<FString>& OutSearchStrings)
 	{
-		if (IsSearchableItemDelegate.IsBound() && !IsSearchableItemDelegate.Execute(Item))
+		if (OverrideGetSearchTerms.IsBound()
+			&& OverrideGetSearchTerms.Execute(Item, OutSearchStrings) == ESearchTermResult::UseOverrideOnly)
 		{
 			return;
 		}
-
+		
 		for (const TPair<FName, TSharedRef<IReplicationTreeColumn<TItemType>>>& ColumnEntry : ColumnInstances)
 		{
 			ColumnEntry.Value->PopulateSearchString(*Item, OutSearchStrings);
@@ -796,6 +837,12 @@ namespace UE::ConcertSharedSlate
 		
 		Items.Sort([this, &IsLessThan](const TSharedPtr<TItemType>& Left, const TSharedPtr<TItemType>& Right)
 		{
+			const EComparisonOverride ComparisonOverride = OverrideIsLessThanDelegate.IsBound() ? OverrideIsLessThanDelegate.Execute(Left, Right) : EComparisonOverride::UseDefault;
+			if (ComparisonOverride != EComparisonOverride::UseDefault)
+			{
+				return ComparisonOverride == EComparisonOverride::Less;
+			}
+			
 			if (PrimarySortInfo.IsValid() && IsLessThan(Left, Right, PrimarySortInfo.SortedColumnId, PrimarySortInfo.SortMode))
 			{
 				return true; // Left comes before Right
diff --git a/Engine/Plugins/Developer/Concert/ConcertUI/ConcertSharedSlate/Source/ConcertSharedSlate/Private/Replication/Editor/View/Tree/TreeItemTraitsInput.h b/Engine/Plugins/Developer/Concert/ConcertUI/ConcertSharedSlate/Source/ConcertSharedSlate/Private/Replication/Editor/View/Tree/TreeItemTraitsInput.h
--- a/Engine/Plugins/Developer/Concert/ConcertUI/ConcertSharedSlate/Source/ConcertSharedSlate/Private/Replication/Editor/View/Tree/TreeItemTraitsInput.h
+++ b/Engine/Plugins/Developer/Concert/ConcertUI/ConcertSharedSlate/Source/ConcertSharedSlate/Private/Replication/Editor/View/Tree/TreeItemTraitsInput.h
@@ -7,6 +7,8 @@
 #include "Types/SlateEnums.h"
 
 class FName;
+class ITableRow;
+class STableViewBase;
 class SWidget;
 
 struct FTableRowStyle;
@@ -49,5 +51,8 @@ namespace UE::ConcertSharedSlate
 			/** Style to use for rows */
 			const FTableRowStyle* RowStyle = nullptr;
 		};
+
+		/** Overrides the widget for the row if the delegate returns non-nullptr. */
+		DECLARE_DELEGATE_RetVal_ThreeParams(TSharedPtr<ITableRow>, FOverrideRowWidget, TSharedPtr<TItemType> Item, const TSharedRef<STableViewBase>& OwnerTable, const FGenerateRowArgs& AdditionalArgs);
 	};
 }
\ No newline at end of file
diff --git a/Engine/Plugins/Developer/Concert/ConcertUI/ConcertSharedSlate/Source/ConcertSharedSlate/Private/Replication/ReplicationWidgetFactories.cpp b/Engine/Plugins/Developer/Concert/ConcertUI/ConcertSharedSlate/Source/ConcertSharedSlate/Private/Replication/ReplicationWidgetFactories.cpp
--- a/Engine/Plugins/Developer/Concert/ConcertUI/ConcertSharedSlate/Source/ConcertSharedSlate/Private/Replication/ReplicationWidgetFactories.cpp
+++ b/Engine/Plugins/Developer/Concert/ConcertUI/ConcertSharedSlate/Source/ConcertSharedSlate/Private/Replication/ReplicationWidgetFactories.cpp
@@ -5,11 +5,15 @@
 #include "Editor/Model/GenericReplicationStreamModel.h"
 #include "Editor/View/MultiEditor/SMultiReplicationStreamEditor.h"
 #include "Editor/View/ObjectEditor/SBaseReplicationStreamEditor.h"
+#include "Editor/View/Property/SMultiObjectAssignment.h"
 #include "Editor/View/Property/SPerObjectPropertyAssignment.h"
 #include "Editor/View/Property/SPropertyTreeView.h"
 #include "Replication/PropertyAssignmentViewFactory.h"
+#include "Replication/Editor/Model/Object/IObjectNameModel.h"
 #include "Replication/Editor/View/IPropertyAssignmentView.h"
 
+#include "Widgets/Text/STextBlock.h"
+
 namespace UE::ConcertSharedSlate
 {
 	TSharedRef<IEditableReplicationStreamModel> CreateBaseStreamModel(
@@ -36,7 +40,7 @@ namespace UE::ConcertSharedSlate
 			.IsEditingEnabled(MoveTemp(EditorParams.IsEditingEnabled))
 			.EditingDisabledToolTipText(MoveTemp(EditorParams.EditingDisabledToolTipText));
 	}
-	
+
 	TSharedRef<IPropertyTreeView> CreateSearchablePropertyTreeView(FCreatePropertyTreeViewParams Params)
 	{
 		// The label column is always required
@@ -60,6 +64,7 @@ namespace UE::ConcertSharedSlate
 		
 		return SNew(SPropertyTreeView)
 			.FilterItem(MoveTemp(FilterDelegate))
+			.CreateCategoryRow(MoveTemp(Params.CreateCategoryRow))
 			.Columns(MoveTemp(Params.PropertyColumns))
 			.ExpandableColumnLabel(ReplicationColumns::Property::LabelColumnId)
 			.PrimarySort(Params.PrimaryPropertySort)
@@ -77,6 +82,13 @@ namespace UE::ConcertSharedSlate
 			.PropertySource(Params.PropertySource);
 	}
 
+	TSharedRef<IMultiObjectPropertyAssignmentView> CreateMultiObjectAssignmentView(FCreateMultiObjectAssignmentViewParams Params)
+	{
+		return SNew(SMultiObjectAssignment, Params.PropertyTreeView)
+			.PropertySource(Params.PropertySource)
+			.ObjectHierarchy(Params.ObjectHierarchy);
+	}
+
 	TSharedRef<IMultiReplicationStreamEditor> CreateBaseMultiStreamEditor(FCreateMultiStreamEditorParams EditorParams, FCreateViewerParams ViewerParams)
 	{
 		return SNew(SMultiReplicationStreamEditor, MoveTemp(EditorParams), MoveTemp(ViewerParams));
diff --git a/Engine/Plugins/Developer/Concert/ConcertUI/ConcertSharedSlate/Source/ConcertSharedSlate/Public/Replication/Editor/Model/Data/PropertyData.h b/Engine/Plugins/Developer/Concert/ConcertUI/ConcertSharedSlate/Source/ConcertSharedSlate/Public/Replication/Editor/Model/Data/PropertyData.h
--- a/Engine/Plugins/Developer/Concert/ConcertUI/ConcertSharedSlate/Source/ConcertSharedSlate/Public/Replication/Editor/Model/Data/PropertyData.h
+++ b/Engine/Plugins/Developer/Concert/ConcertUI/ConcertSharedSlate/Source/ConcertSharedSlate/Public/Replication/Editor/Model/Data/PropertyData.h
@@ -14,13 +14,13 @@ namespace UE::ConcertSharedSlate
 	{
 	public:
 		
-		FPropertyData(TSet<TSoftObjectPtr<>> ContextObjects, FSoftClassPath OwningClass, FConcertPropertyChain Object)
+		FPropertyData(TArray<TSoftObjectPtr<>> ContextObjects, FSoftClassPath OwningClass, FConcertPropertyChain Object)
 			: ContextObjects(MoveTemp(ContextObjects))
 			, OwningClassPtr(MoveTemp(OwningClass))
 			, Property(MoveTemp(Object))
 		{}
 		
-		const TSet<TSoftObjectPtr<>>& GetContextObjects() const { return ContextObjects; }
+		const TArray<TSoftObjectPtr<>>& GetContextObjects() const { return ContextObjects; }
 		const FConcertPropertyChain& GetProperty() const { return Property; }
 		const TSoftClassPtr<>& GetOwningClassPtr() const { return OwningClassPtr; }
 
@@ -33,7 +33,7 @@ namespace UE::ConcertSharedSlate
 		 * This has multiple elements in the case of multi-edit (i.e. when the user clicks multiple, compatible actors in the top-view).
 		 * For example, for multi-edit this could contain ActorA->StaticMeshComponent0 and ActorB->StaticMeshComponent0.
 		 */
-		TSet<TSoftObjectPtr<>> ContextObjects;
+		TArray<TSoftObjectPtr<>> ContextObjects;
 		
 		/** The class with which the FProperty can be determined. */
 		TSoftClassPtr<> OwningClassPtr;
diff --git a/Engine/Plugins/Developer/Concert/ConcertUI/ConcertSharedSlate/Source/ConcertSharedSlate/Public/Replication/Editor/View/Extension/CategoryRowGeneration.h b/Engine/Plugins/Developer/Concert/ConcertUI/ConcertSharedSlate/Source/ConcertSharedSlate/Public/Replication/Editor/View/Extension/CategoryRowGeneration.h
new file mode 100644
--- /dev/null
+++ b/Engine/Plugins/Developer/Concert/ConcertUI/ConcertSharedSlate/Source/ConcertSharedSlate/Public/Replication/Editor/View/Extension/CategoryRowGeneration.h
@@ -0,0 +1,50 @@
+﻿// Copyright Epic Games, Inc. All Rights Reserved.
+
+#pragma once
+
+#include "Containers/Array.h"
+#include "Misc/Attribute.h"
+#include "Templates/SharedPointer.h"
+#include "UObject/SoftObjectPtr.h"
+
+class FString;
+class FText;
+class SWidget;
+
+namespace UE::ConcertSharedSlate
+{
+	class FPropertyData;
+	class IObjectNameModel;
+	
+	struct FCategoryRowGenerationArgs
+	{
+		/** The objects for which this is being generated */
+		TArray<TSoftObjectPtr<>> ContextObjects;
+		/** The highlight text which contains the search terms*/
+		TAttribute<FText> HighlightText;
+
+		FCategoryRowGenerationArgs(TArray<TSoftObjectPtr<>> ContextObjects, TAttribute<FText> HighlightText)
+			: ContextObjects(MoveTemp(ContextObjects))
+			, HighlightText(MoveTemp(HighlightText))
+		{}
+	};
+	
+	/** A special reassignment view which shows the display object and all its children. */
+	class CONCERTSHAREDSLATE_API ICategoryRow
+	{
+	public:
+
+		/** Generates search terms */
+		virtual void GenerateSearchTerms(const TArray<TSoftObjectPtr<>>& ContextObjects, TArray<FString>& InOutSearchTerms) const = 0;
+
+		/** Gets the row as widget */
+		virtual TSharedRef<SWidget> GetWidget() = 0;
+
+		virtual ~ICategoryRow() = default;
+	};
+	
+	DECLARE_DELEGATE_RetVal_OneParam(TSharedRef<ICategoryRow>, FCreateCategoryRow, const FCategoryRowGenerationArgs& Args);
+	
+	/** Creates a delegate that generates a row displaying the name of the first context object. */
+	CONCERTSHAREDSLATE_API FCreateCategoryRow CreateDefaultCategoryGenerator(TSharedRef<IObjectNameModel> NameModel);
+}
diff --git a/Engine/Plugins/Developer/Concert/ConcertUI/ConcertSharedSlate/Source/ConcertSharedSlate/Public/Replication/Editor/View/IMultiObjectPropertyAssignmentView.h b/Engine/Plugins/Developer/Concert/ConcertUI/ConcertSharedSlate/Source/ConcertSharedSlate/Public/Replication/Editor/View/IMultiObjectPropertyAssignmentView.h
new file mode 100644
--- /dev/null
+++ b/Engine/Plugins/Developer/Concert/ConcertUI/ConcertSharedSlate/Source/ConcertSharedSlate/Public/Replication/Editor/View/IMultiObjectPropertyAssignmentView.h
@@ -0,0 +1,21 @@
+﻿// Copyright Epic Games, Inc. All Rights Reserved.
+
+#pragma once
+
+#include "IPropertyAssignmentView.h"
+
+namespace UE::ConcertSharedSlate
+{
+	/** A special reassignment view which shows the display object and all its children. */
+	class CONCERTSHAREDSLATE_API IMultiObjectPropertyAssignmentView : public IPropertyAssignmentView
+	{
+	public:
+
+		/** Sets whether subobjects should be shown in the view (applies only to subobjects other than UActorComponents, which are always shown). */
+		virtual void SetShouldShowSubobjects(bool bShowSubobjects) = 0;
+		/** @return Whether subobjects should be shown */
+		virtual bool GetShouldShowSubobjects() const = 0;
+
+		virtual ~IMultiObjectPropertyAssignmentView() = default;
+	};
+}
diff --git a/Engine/Plugins/Developer/Concert/ConcertUI/ConcertSharedSlate/Source/ConcertSharedSlate/Public/Replication/Editor/View/IPropertyAssignmentView.h b/Engine/Plugins/Developer/Concert/ConcertUI/ConcertSharedSlate/Source/ConcertSharedSlate/Public/Replication/Editor/View/IPropertyAssignmentView.h
--- a/Engine/Plugins/Developer/Concert/ConcertUI/ConcertSharedSlate/Source/ConcertSharedSlate/Public/Replication/Editor/View/IPropertyAssignmentView.h
+++ b/Engine/Plugins/Developer/Concert/ConcertUI/ConcertSharedSlate/Source/ConcertSharedSlate/Public/Replication/Editor/View/IPropertyAssignmentView.h
@@ -35,7 +35,7 @@ namespace UE::ConcertSharedSlate
 		 * @param Objects The objects that are supposed to be displayed.
 		 * @param Model The model that can be queried for object info.
 		 */
-		virtual void RefreshData(const TSet<FSoftObjectPath>& Objects, const IReplicationStreamModel& Model) = 0;
+		virtual void RefreshData(const TArray<FSoftObjectPath>& Objects, const IReplicationStreamModel& Model) = 0;
 		
 		/** Reapply the filter function to all items at the end of the frame. Call e.g. when the filters have changed. */
 		virtual void RequestRefilter() const = 0;
diff --git a/Engine/Plugins/Developer/Concert/ConcertUI/ConcertSharedSlate/Source/ConcertSharedSlate/Public/Replication/Editor/View/IPropertyTreeView.h b/Engine/Plugins/Developer/Concert/ConcertUI/ConcertSharedSlate/Source/ConcertSharedSlate/Public/Replication/Editor/View/IPropertyTreeView.h
--- a/Engine/Plugins/Developer/Concert/ConcertUI/ConcertSharedSlate/Source/ConcertSharedSlate/Public/Replication/Editor/View/IPropertyTreeView.h
+++ b/Engine/Plugins/Developer/Concert/ConcertUI/ConcertSharedSlate/Source/ConcertSharedSlate/Public/Replication/Editor/View/IPropertyTreeView.h
@@ -22,7 +22,7 @@ namespace UE::ConcertSharedSlate
 		 * This has multiple elements in the case of multi-edit (i.e. when the user clicks multiple, compatible actors in the top-view).
 		 * For example, for multi-edit this could contain ActorA->StaticMeshComponent0 and ActorB->StaticMeshComponent0.
 		 */
-		TSet<TSoftObjectPtr<>> ContextObjects;
+		TArray<TSoftObjectPtr<>> ContextObjects;
 		
 		/** The properties to display */
 		TSet<FConcertPropertyChain> PropertiesToDisplay;
diff --git a/Engine/Plugins/Developer/Concert/ConcertUI/ConcertSharedSlate/Source/ConcertSharedSlate/Public/Replication/Editor/View/IReplicationStreamViewer.h b/Engine/Plugins/Developer/Concert/ConcertUI/ConcertSharedSlate/Source/ConcertSharedSlate/Public/Replication/Editor/View/IReplicationStreamViewer.h
--- a/Engine/Plugins/Developer/Concert/ConcertUI/ConcertSharedSlate/Source/ConcertSharedSlate/Public/Replication/Editor/View/IReplicationStreamViewer.h
+++ b/Engine/Plugins/Developer/Concert/ConcertUI/ConcertSharedSlate/Source/ConcertSharedSlate/Public/Replication/Editor/View/IReplicationStreamViewer.h
@@ -34,7 +34,7 @@ namespace UE::ConcertSharedSlate
 		 * If there is an IReplicationSubobjectView, this is IReplicationSubobjectView::GetSelectedObjects.
 		 * Otherwise it is IReplicationStreamViewer::GetSelectedTopLevelObjects.
 		 */
-		virtual TSet<FSoftObjectPath> GetObjectsBeingPropertyEdited() const = 0;
+		virtual TArray<FSoftObjectPath> GetObjectsBeingPropertyEdited() const = 0;
 
 		virtual ~IReplicationStreamViewer() = default;
 	};
diff --git a/Engine/Plugins/Developer/Concert/ConcertUI/ConcertSharedSlate/Source/ConcertSharedSlate/Public/Replication/PropertyAssignmentViewFactory.h b/Engine/Plugins/Developer/Concert/ConcertUI/ConcertSharedSlate/Source/ConcertSharedSlate/Public/Replication/PropertyAssignmentViewFactory.h
--- a/Engine/Plugins/Developer/Concert/ConcertUI/ConcertSharedSlate/Source/ConcertSharedSlate/Public/Replication/PropertyAssignmentViewFactory.h
+++ b/Engine/Plugins/Developer/Concert/ConcertUI/ConcertSharedSlate/Source/ConcertSharedSlate/Public/Replication/PropertyAssignmentViewFactory.h
@@ -12,6 +12,8 @@
 
 namespace UE::ConcertSharedSlate
 {
+	class IObjectHierarchyModel;
+	class IMultiObjectPropertyAssignmentView;
 	class FPropertyData;
 	class IPropertyAssignmentView;
 	class IPropertySelectionSourceModel;
@@ -39,4 +41,30 @@ namespace UE::ConcertSharedSlate
 	 * You can customize the tree view by injecting columns (@see CreateSearchablePropertyTreeView).
 	 */
 	CONCERTSHAREDSLATE_API TSharedRef<IPropertyAssignmentView> CreatePerObjectAssignmentView(FCreatePerObjectAssignmentViewParams Params = {});
+	
+	struct FCreateMultiObjectAssignmentViewParams : FCreatePerObjectAssignmentViewParams
+	{
+		/** Required. Displays the properties in a tree view. You can pass in e.g. custom UI with advanced filtering. */
+		TSharedRef<IPropertyTreeView> PropertyTreeView = CreateSearchablePropertyTreeView();
+
+		/**
+		 * Optional.
+		 * Gets components and subobjects of the displayed object.
+		 * If this is unspecified, the created IMultiObjectPropertyAssignmentView will behave exactly as the per object view (IPropertyAssignmentView).
+		 */
+		TSharedPtr<IObjectHierarchyModel> ObjectHierarchy;
+		
+		/**
+		 * Optional.
+		 * If specified, the view will display all properties reported by the model (useful for editor UI which edits streams, not useful for server where property info is not available).
+		 * If unspecified, only display the properties assigned to the object in the stream.
+		 */
+		TSharedPtr<IPropertySelectionSourceModel> PropertySource;
+	};
+
+	/**
+	 * Creates a view that shows the properties of the object the user clicks and its subobjects.
+	 * You can customize the tree view by injecting columns (@see CreateSearchablePropertyTreeView).
+	 */
+	CONCERTSHAREDSLATE_API TSharedRef<IMultiObjectPropertyAssignmentView> CreateMultiObjectAssignmentView(FCreateMultiObjectAssignmentViewParams Params = {});
 }
diff --git a/Engine/Plugins/Developer/Concert/ConcertUI/ConcertSharedSlate/Source/ConcertSharedSlate/Public/Replication/PropertyTreeFactory.h b/Engine/Plugins/Developer/Concert/ConcertUI/ConcertSharedSlate/Source/ConcertSharedSlate/Public/Replication/PropertyTreeFactory.h
--- a/Engine/Plugins/Developer/Concert/ConcertUI/ConcertSharedSlate/Source/ConcertSharedSlate/Public/Replication/PropertyTreeFactory.h
+++ b/Engine/Plugins/Developer/Concert/ConcertUI/ConcertSharedSlate/Source/ConcertSharedSlate/Public/Replication/PropertyTreeFactory.h
@@ -5,15 +5,21 @@
 #include "Editor/View/Column/IPropertyTreeColumn.h"
 #include "Editor/View/Column/ReplicationColumnsUtils.h"
 #include "Editor/View/Column/SelectionViewerColumns.h"
+#include "Editor/View/Extension/CategoryRowGeneration.h"
 #include "Replication/Utils/FilterResult.h"
 
+#include "Containers/ContainersFwd.h"
 #include "Delegates/Delegate.h"
 #include "Misc/Attribute.h"
 #include "Templates/SharedPointer.h"
+#include "UObject/SoftObjectPtr.h"
+
+class SExpanderArrow;
 
 namespace UE::ConcertSharedSlate
 {
 	class FPropertyData;
+	class IObjectNameModel;
 	class IPropertyTreeView;
 	
 	DECLARE_DELEGATE_RetVal_OneParam(EFilterResult, FFilterPropertyData, const FPropertyData&);
@@ -26,8 +32,17 @@ namespace UE::ConcertSharedSlate
 			ReplicationColumns::Property::LabelColumn()
 		};
 
-		/** Optional filter function. Return true to al */
+		/** Optional filter function. */
 		FFilterPropertyData FilterItem;
+
+		/**
+		 * Optional delegate for grouping objects under a category.
+		 * If unset, no category are generated.
+		 * 
+		 * When the user clicks an object in the top view, this delegate will be called for the clicked object, its components (if an actor), and its (nested) subobjects.
+		 * ContextObjects is a single object if a single object is clicked or multiple object in the case of multi-edit.
+		 */
+		FCreateCategoryRow CreateCategoryRow;
 		
 		/** Optional initial primary sort mode for object rows */
 		FColumnSortInfo PrimaryPropertySort { ReplicationColumns::Property::LabelColumnId, EColumnSortMode::Ascending };
