Refactor IPropertySelectionSourceModel:.GetPropertySource signature to also accept TSoftObjectPtr. This is in preparation for an upcoming change where the user can select properties to display in the UI.
Also refactors codebase to use TSoftObjectPtr<> in favour of FSoftObjectPath.

diff --git a/Engine/Plugins/Developer/Concert/ConcertApp/MultiUserClient/Source/MultiUserClient/Private/Widgets/ActiveSession/Replication/Client/Context/ContextMenuUtils.h b/Engine/Plugins/Developer/Concert/ConcertApp/MultiUserClient/Source/MultiUserClient/Private/Widgets/ActiveSession/Replication/Client/Context/ContextMenuUtils.h
--- a/Engine/Plugins/Developer/Concert/ConcertApp/MultiUserClient/Source/MultiUserClient/Private/Widgets/ActiveSession/Replication/Client/Context/ContextMenuUtils.h
+++ b/Engine/Plugins/Developer/Concert/ConcertApp/MultiUserClient/Source/MultiUserClient/Private/Widgets/ActiveSession/Replication/Client/Context/ContextMenuUtils.h
@@ -4,6 +4,7 @@
 
 #include "Containers/ArrayView.h"
 #include "UObject/SoftObjectPath.h"
+#include "UObject/SoftObjectPtr.h"
 
 class IConcertClient;
 class FMenuBuilder;
@@ -24,7 +25,7 @@ namespace UE::MultiUserClient::ContextMenuUtils
 	/** Adds menu entries for reassigning the object to another client. */
 	void AddReassignmentOptions(
 		FMenuBuilder& MenuBuilder,
-		const FSoftObjectPath& ContextObject,
+		const TSoftObjectPtr<>& ContextObject,
 		const IConcertClient& ConcertClient,
 		const FReplicationClientManager& ReplicationManager,
 		ConcertSharedSlate::IObjectHierarchyModel& ObjectHierarchy,
@@ -35,11 +36,11 @@ namespace UE::MultiUserClient::ContextMenuUtils
 	/** Adds an edit box for batch reassigning the select object's frequencies for all replicating clients. */
 	void AddFrequencyOptionsForMultipleClients(FMenuBuilder& MenuBuilder, const FSoftObjectPath& ContextObjects, FReplicationClientManager& InClientManager);
 	/** Adds an edit box for batch reassigning the select object's frequencies for all replicating clients. */
-	inline void AddFrequencyOptionsIfOneContextObject_MultiClient(FMenuBuilder& MenuBuilder, TConstArrayView<FSoftObjectPath> ContextObjects, FReplicationClientManager& InClientManager)
+	inline void AddFrequencyOptionsIfOneContextObject_MultiClient(FMenuBuilder& MenuBuilder, TConstArrayView<TSoftObjectPtr<>> ContextObjects, FReplicationClientManager& InClientManager)
 	{
 		if (ContextObjects.Num() == 1)
 		{
-			AddFrequencyOptionsForMultipleClients(MenuBuilder, ContextObjects[0], InClientManager);
+			AddFrequencyOptionsForMultipleClients(MenuBuilder, ContextObjects[0].GetUniqueID(), InClientManager);
 		}
 	}
 }
\ No newline at end of file
diff --git a/Engine/Plugins/Developer/Concert/ConcertApp/MultiUserClient/Source/MultiUserClient/Private/Widgets/ActiveSession/Replication/Client/Context/ContextMenuUtils_Reassignment.cpp b/Engine/Plugins/Developer/Concert/ConcertApp/MultiUserClient/Source/MultiUserClient/Private/Widgets/ActiveSession/Replication/Client/Context/ContextMenuUtils_Reassignment.cpp
--- a/Engine/Plugins/Developer/Concert/ConcertApp/MultiUserClient/Source/MultiUserClient/Private/Widgets/ActiveSession/Replication/Client/Context/ContextMenuUtils_Reassignment.cpp
+++ b/Engine/Plugins/Developer/Concert/ConcertApp/MultiUserClient/Source/MultiUserClient/Private/Widgets/ActiveSession/Replication/Client/Context/ContextMenuUtils_Reassignment.cpp
@@ -23,9 +23,14 @@ namespace UE::MultiUserClient::ContextMenuUtils
 		using FInlineAllocator = TInlineAllocator<24>;
 		using FInlineObjectPathArray = TArray<FSoftObjectPath, FInlineAllocator>;
 		
-		static FInlineObjectPathArray GetChildrenOfManagedObject(ConcertSharedSlate::IObjectHierarchyModel& ObjectHierarchy, FSoftObjectPath ManagedObject)
+		static FInlineObjectPathArray GetChildrenOfManagedObject(ConcertSharedSlate::IObjectHierarchyModel& ObjectHierarchy, TSoftObjectPtr<> ManagedObject)
 		{
-			return ObjectHierarchy.GetChildrenRecursive<FInlineAllocator>(ManagedObject);
+			FInlineObjectPathArray Result;
+			Algo::Transform(ObjectHierarchy.GetChildrenRecursive<FInlineAllocator>(ManagedObject), Result, [](const TSoftObjectPtr<>& Object)
+			{
+				return Object.GetUniqueID();
+			});
+			return Result;
 		}
 		
 		static void AddReassignSection(
@@ -69,7 +74,7 @@ namespace UE::MultiUserClient::ContextMenuUtils
 	
 	void AddReassignmentOptions(
 		FMenuBuilder& MenuBuilder,
-		const FSoftObjectPath& ContextObject,
+		const TSoftObjectPtr<>& ContextObject,
 		const IConcertClient& ConcertClient,
 		const FReplicationClientManager& ReplicationManager,
 		ConcertSharedSlate::IObjectHierarchyModel& ObjectHierarchy,
@@ -80,7 +85,7 @@ namespace UE::MultiUserClient::ContextMenuUtils
 		const TArray<const FReplicationClient*> SortedClients = ClientUtils::GetSortedClientList(ConcertClient, ReplicationManager);
 		
 		MenuBuilder.BeginSection(TEXT("Reassign.This"), LOCTEXT("Reassign.This", "Reassign this to"));
-		Private::AddReassignSection(MenuBuilder, SortedClients, Private::FInlineObjectPathArray{ ContextObject }, ConcertClient, ReassignmentLogic, MultiStreamEditor);
+		Private::AddReassignSection(MenuBuilder, SortedClients, Private::FInlineObjectPathArray{ ContextObject.GetUniqueID() }, ConcertClient, ReassignmentLogic, MultiStreamEditor);
 		MenuBuilder.EndSection();
 
 		// Do not distract the user with more options if children have no assigned properties
diff --git a/Engine/Plugins/Developer/Concert/ConcertApp/MultiUserClient/Source/MultiUserClient/Private/Widgets/ActiveSession/Replication/Client/Multi/Columns/MultiStreamColumns_AssignProperties.cpp b/Engine/Plugins/Developer/Concert/ConcertApp/MultiUserClient/Source/MultiUserClient/Private/Widgets/ActiveSession/Replication/Client/Multi/Columns/MultiStreamColumns_AssignProperties.cpp
--- a/Engine/Plugins/Developer/Concert/ConcertApp/MultiUserClient/Source/MultiUserClient/Private/Widgets/ActiveSession/Replication/Client/Multi/Columns/MultiStreamColumns_AssignProperties.cpp
+++ b/Engine/Plugins/Developer/Concert/ConcertApp/MultiUserClient/Source/MultiUserClient/Private/Widgets/ActiveSession/Replication/Client/Multi/Columns/MultiStreamColumns_AssignProperties.cpp
@@ -31,9 +31,9 @@ namespace UE::MultiUserClient::MultiStreamColumns
 		{
 			for (const TSharedRef<ConcertSharedSlate::IEditableReplicationStreamModel>& Stream : MultiEditor.GetMultiStreamModel().GetEditableStreams())
 			{
-				for (const FSoftObjectPath& SelectedObject : MultiEditor.GetEditorBase().GetObjectsBeingPropertyEdited())
+				for (const TSoftObjectPtr<>& SelectedObject : MultiEditor.GetEditorBase().GetSelectedObjects())
 				{
-					const bool bStreamAssignedToProperty = Stream->HasProperty(SelectedObject, Property);
+					const bool bStreamAssignedToProperty = Stream->HasProperty(SelectedObject.GetUniqueID(), Property);
 					if (bStreamAssignedToProperty)
 					{
 						Consume(Stream);
@@ -106,7 +106,7 @@ namespace UE::MultiUserClient::MultiStreamColumns
 			
 			virtual TSharedRef<SWidget> GenerateColumnWidget(const FBuildArgs& InArgs) override
 			{
-				const TArray<FSoftObjectPath> DisplayedObjects = MultiStreamEditor.Get()->GetEditorBase().GetObjectsBeingPropertyEdited();
+				const TArray<TSoftObjectPtr<>> DisplayedObjects = MultiStreamEditor.Get()->GetEditorBase().GetSelectedObjects();
 				return SNew(SAssignPropertyComboBox, MultiStreamEditor.Get().ToSharedRef(), ConcertClient, ClientManager)
 					.DisplayedProperty(InArgs.RowItem.RowData.GetProperty())
 					.EditedObjects(DisplayedObjects)
@@ -132,7 +132,7 @@ namespace UE::MultiUserClient::MultiStreamColumns
 			virtual bool CanBeSorted() const override { return true; }
 			virtual bool IsLessThan(const FPropertyTreeRowContext& Left, const FPropertyTreeRowContext& Right) const override
 			{
-				const TArray<FSoftObjectPath> DisplayedObjects = MultiStreamEditor.Get()->GetEditorBase().GetObjectsBeingPropertyEdited();
+				const TArray<TSoftObjectPtr<>> DisplayedObjects = MultiStreamEditor.Get()->GetEditorBase().GetSelectedObjects();
 				const TOptional<FString> LeftClientDisplayString = SAssignPropertyComboBox::GetDisplayString(ConcertClient, ClientManager, Left.RowData.GetProperty(), DisplayedObjects);
 				const TOptional<FString> RightClientDisplayString = SAssignPropertyComboBox::GetDisplayString(ConcertClient, ClientManager, Right.RowData.GetProperty(), DisplayedObjects);
 			
diff --git a/Engine/Plugins/Developer/Concert/ConcertApp/MultiUserClient/Source/MultiUserClient/Private/Widgets/ActiveSession/Replication/Client/Multi/Columns/SAssignPropertyComboBox.cpp b/Engine/Plugins/Developer/Concert/ConcertApp/MultiUserClient/Source/MultiUserClient/Private/Widgets/ActiveSession/Replication/Client/Multi/Columns/SAssignPropertyComboBox.cpp
--- a/Engine/Plugins/Developer/Concert/ConcertApp/MultiUserClient/Source/MultiUserClient/Private/Widgets/ActiveSession/Replication/Client/Multi/Columns/SAssignPropertyComboBox.cpp
+++ b/Engine/Plugins/Developer/Concert/ConcertApp/MultiUserClient/Source/MultiUserClient/Private/Widgets/ActiveSession/Replication/Client/Multi/Columns/SAssignPropertyComboBox.cpp
@@ -30,15 +30,15 @@ namespace UE::MultiUserClient
 {
 	namespace AssignPropertyComboBox
 	{
-		TArray<FGuid> GetDisplayedClients(const FReplicationClientManager& ClientManager, const FConcertPropertyChain& DisplayedProperty, const TArray<FSoftObjectPath>& EditedObjects)
+		TArray<FGuid> GetDisplayedClients(const FReplicationClientManager& ClientManager, const FConcertPropertyChain& DisplayedProperty, const TArray<TSoftObjectPtr<>>& EditedObjects)
 		{
 			TArray<FGuid> Clients;
 			ClientManager.ForEachClient([&DisplayedProperty, &EditedObjects, &Clients](const FReplicationClient& Client)
 			{
 				const TMap<FSoftObjectPath, FConcertReplicatedObjectInfo>& ObjectInfoMap = Client.GetStreamSynchronizer().GetServerState().ReplicatedObjects;
-				for (const FSoftObjectPath& ObjectPath : EditedObjects)
+				for (const TSoftObjectPtr<>& ObjectPath : EditedObjects)
 				{
-					if (const FConcertReplicatedObjectInfo* ObjectInfo = ObjectInfoMap.Find(ObjectPath)
+					if (const FConcertReplicatedObjectInfo* ObjectInfo = ObjectInfoMap.Find(ObjectPath.GetUniqueID())
 						; ObjectInfo && ObjectInfo->PropertySelection.ReplicatedProperties.Contains(DisplayedProperty))
 					{
 						Clients.Add(Client.GetEndpointId());
@@ -55,7 +55,7 @@ namespace UE::MultiUserClient
 		const TSharedRef<IConcertClient>& LocalConcertClient,
 		const FReplicationClientManager& ClientManager,
 		const FConcertPropertyChain& DisplayedProperty,
-		const TArray<FSoftObjectPath>& EditedObjects)
+		const TArray<TSoftObjectPtr<>>& EditedObjects)
 	{
 		using SWidgetType = ConcertClientSharedSlate::SHorizontalClientList;
 		const TArray<FGuid> Clients = AssignPropertyComboBox::GetDisplayedClients(ClientManager, DisplayedProperty, EditedObjects);
@@ -212,15 +212,16 @@ namespace UE::MultiUserClient
 
 			// ... and then assign the property
 			const TSharedRef<ConcertSharedSlate::IEditableReplicationStreamModel> EditModel = Client->GetClientEditModel();
-			for (const FSoftObjectPath& ObjectPath : EditedObjects)
+			for (const TSoftObjectPtr<>& Object : EditedObjects)
 			{
+				const FSoftObjectPath& ObjectPath = Object.GetUniqueID();
 				if (!EditModel->ContainsObjects({ ObjectPath }))
 				{
-					EditModel->AddObjects({ ObjectPath.ResolveObject() });
+					EditModel->AddObjects({ Object.Get() });
 				}
 
 				const FSoftClassPath ClassPath = EditModel->GetObjectClass(ObjectPath);
-				TArray<FConcertPropertyChain> AddedProperties { Property };
+				TArray AddedProperties { Property };
 				ConcertClientSharedSlate::PropertyUtils::AppendAdditionalPropertiesToAdd(ClassPath, AddedProperties);
 				EditModel->AddProperties(ObjectPath, AddedProperties);
 			}
@@ -246,9 +247,9 @@ namespace UE::MultiUserClient
 		{
 			if (*Client != ClientToRemoveFrom && !ClientToRemoveFrom.AllowsEditing())
 			{
-				const bool bHasAnySelectedObject = Algo::AnyOf(EditedObjects, [this, &ClientToRemoveFrom](const FSoftObjectPath& ObjectPath)
+				const bool bHasAnySelectedObject = Algo::AnyOf(EditedObjects, [this, &ClientToRemoveFrom](const TSoftObjectPtr<>& Object)
 				{
-					return ClientToRemoveFrom.GetClientEditModel()->HasProperty(ObjectPath, Property);
+					return ClientToRemoveFrom.GetClientEditModel()->HasProperty(Object.GetUniqueID(), Property);
 				});
 				bCanRemoveFromOwners = !bHasAnySelectedObject;
 				
@@ -289,9 +290,9 @@ namespace UE::MultiUserClient
 
 		const TSharedRef<ConcertSharedSlate::IEditableReplicationStreamModel> Model = Client->GetClientEditModel();
 		ECheckBoxState CheckBoxState = ECheckBoxState::Undetermined;
-		for (const FSoftObjectPath& ObjectPath : EditedObjects)
+		for (const TSoftObjectPtr<>& ObjectPath : EditedObjects)
 		{
-			const bool bHasProperty = Model->HasProperty(ObjectPath, Property);
+			const bool bHasProperty = Model->HasProperty(ObjectPath.GetUniqueID(), Property);
 			switch (CheckBoxState)
 			{
 			case ECheckBoxState::Unchecked:
@@ -330,7 +331,7 @@ namespace UE::MultiUserClient
 		bool bIsAssignedToAnyClient = false;
 		ClientManager->ForEachClient([this, &bIsAssignedToAnyClient](const FReplicationClient& Client)
 		{
-			for (const FSoftObjectPath& EditedObject : EditedObjects)
+			for (const TSoftObjectPtr<>& EditedObject : EditedObjects)
 			{
 				if (bIsAssignedToAnyClient)
 				{
@@ -338,7 +339,7 @@ namespace UE::MultiUserClient
 				}
 				
 				const TSharedRef<ConcertSharedSlate::IEditableReplicationStreamModel> Model = Client.GetClientEditModel();
-				const bool bHasProperty = Model->HasProperty(EditedObject, Property);
+				const bool bHasProperty = Model->HasProperty(EditedObject.GetUniqueID(), Property);
 				bIsAssignedToAnyClient |= bHasProperty;
 			}
 			
@@ -354,8 +355,9 @@ namespace UE::MultiUserClient
 			const TSharedRef<ConcertSharedSlate::IEditableReplicationStreamModel> EditModel = ClientToRemoveFrom.GetClientEditModel();
 			if (ClientToRemoveFrom.AllowsEditing() && ShouldRemoveFromClient(ClientToRemoveFrom))
 			{
-				for (const FSoftObjectPath& ObjectPath : EditedObjects)
+				for (const TSoftObjectPtr<>& Object : EditedObjects)
 				{
+					const FSoftObjectPath& ObjectPath = Object.GetUniqueID();
 					const FSoftClassPath ClassPath = EditModel->GetObjectClass(ObjectPath);
 					EditModel->RemoveProperties(ObjectPath, { Property });
 					
diff --git a/Engine/Plugins/Developer/Concert/ConcertApp/MultiUserClient/Source/MultiUserClient/Private/Widgets/ActiveSession/Replication/Client/Multi/Columns/SAssignPropertyComboBox.h b/Engine/Plugins/Developer/Concert/ConcertApp/MultiUserClient/Source/MultiUserClient/Private/Widgets/ActiveSession/Replication/Client/Multi/Columns/SAssignPropertyComboBox.h
--- a/Engine/Plugins/Developer/Concert/ConcertApp/MultiUserClient/Source/MultiUserClient/Private/Widgets/ActiveSession/Replication/Client/Multi/Columns/SAssignPropertyComboBox.h
+++ b/Engine/Plugins/Developer/Concert/ConcertApp/MultiUserClient/Source/MultiUserClient/Private/Widgets/ActiveSession/Replication/Client/Multi/Columns/SAssignPropertyComboBox.h
@@ -38,13 +38,13 @@ namespace UE::MultiUserClient
 			const TSharedRef<IConcertClient>& LocalConcertClient,
 			const FReplicationClientManager& ClientManager,
 			const FConcertPropertyChain& DisplayedProperty,
-			const TArray<FSoftObjectPath>& EditedObjects
+			const TArray<TSoftObjectPtr<>>& EditedObjects
 			);
 		
 		SLATE_BEGIN_ARGS(SAssignPropertyComboBox)
 		{}
 			SLATE_ARGUMENT(FConcertPropertyChain, DisplayedProperty)
-			SLATE_ARGUMENT(TArray<FSoftObjectPath>, EditedObjects)
+			SLATE_ARGUMENT(TArray<TSoftObjectPtr<>>, EditedObjects)
 			SLATE_ARGUMENT(TSharedPtr<FText>, HighlightText)
 			/** Called when the property assignment of any client(s) is changed by this widget. */
 			SLATE_EVENT(FOnPropertyAssignmentChanged, OnPropertyAssignmentChanged)
@@ -65,7 +65,7 @@ namespace UE::MultiUserClient
 		TSharedPtr<IConcertClient> ConcertClient;
 
 		/** The objects for which the property is being displayed. */
-		TArray<FSoftObjectPath> EditedObjects;
+		TArray<TSoftObjectPtr<>> EditedObjects;
 		/** The property assigned to this column. */
 		FConcertPropertyChain Property;
 
diff --git a/Engine/Plugins/Developer/Concert/ConcertApp/MultiUserClient/Source/MultiUserClient/Private/Widgets/ActiveSession/Replication/Client/Multi/SMultiClientView.cpp b/Engine/Plugins/Developer/Concert/ConcertApp/MultiUserClient/Source/MultiUserClient/Private/Widgets/ActiveSession/Replication/Client/Multi/SMultiClientView.cpp
--- a/Engine/Plugins/Developer/Concert/ConcertApp/MultiUserClient/Source/MultiUserClient/Private/Widgets/ActiveSession/Replication/Client/Multi/SMultiClientView.cpp
+++ b/Engine/Plugins/Developer/Concert/ConcertApp/MultiUserClient/Source/MultiUserClient/Private/Widgets/ActiveSession/Replication/Client/Multi/SMultiClientView.cpp
@@ -205,7 +205,7 @@ namespace UE::MultiUserClient
 		StreamEditor->GetEditorBase().Refresh();
 	}
 
-	void SMultiClientView::ExtendObjectContextMenu(FMenuBuilder& MenuBuilder, TConstArrayView<FSoftObjectPath> ContextObjects) const
+	void SMultiClientView::ExtendObjectContextMenu(FMenuBuilder& MenuBuilder, TConstArrayView<TSoftObjectPtr<>> ContextObjects) const
 	{
 		ContextMenuUtils::AddFrequencyOptionsIfOneContextObject_MultiClient(MenuBuilder, ContextObjects, *ClientManager);
 
diff --git a/Engine/Plugins/Developer/Concert/ConcertApp/MultiUserClient/Source/MultiUserClient/Private/Widgets/ActiveSession/Replication/Client/Multi/SMultiClientView.h b/Engine/Plugins/Developer/Concert/ConcertApp/MultiUserClient/Source/MultiUserClient/Private/Widgets/ActiveSession/Replication/Client/Multi/SMultiClientView.h
--- a/Engine/Plugins/Developer/Concert/ConcertApp/MultiUserClient/Source/MultiUserClient/Private/Widgets/ActiveSession/Replication/Client/Multi/SMultiClientView.h
+++ b/Engine/Plugins/Developer/Concert/ConcertApp/MultiUserClient/Source/MultiUserClient/Private/Widgets/ActiveSession/Replication/Client/Multi/SMultiClientView.h
@@ -69,7 +69,7 @@ namespace UE::MultiUserClient
 		void OnHierarchyNeedsRefresh() const;
 		
 		/** Adds additional entries to the context menu for the object tree view. */
-		void ExtendObjectContextMenu(FMenuBuilder& MenuBuilder, TConstArrayView<FSoftObjectPath> ContextObjects) const;
+		void ExtendObjectContextMenu(FMenuBuilder& MenuBuilder, TConstArrayView<TSoftObjectPtr<>> ContextObjects) const;
 		/** Decides whether the object should be displayed: do not show it if it's not in the editor world. */
 		bool ShouldDisplayObject(const FSoftObjectPath& Object) const;
 	};
diff --git a/Engine/Plugins/Developer/Concert/ConcertUI/ConcertClientSharedSlate/Source/ConcertClientSharedSlate/Private/Replication/Editor/Model/Object/EditorObjectHierarchyModel.cpp b/Engine/Plugins/Developer/Concert/ConcertUI/ConcertClientSharedSlate/Source/ConcertClientSharedSlate/Private/Replication/Editor/Model/Object/EditorObjectHierarchyModel.cpp
--- a/Engine/Plugins/Developer/Concert/ConcertUI/ConcertClientSharedSlate/Source/ConcertClientSharedSlate/Private/Replication/Editor/Model/Object/EditorObjectHierarchyModel.cpp
+++ b/Engine/Plugins/Developer/Concert/ConcertUI/ConcertClientSharedSlate/Source/ConcertClientSharedSlate/Private/Replication/Editor/Model/Object/EditorObjectHierarchyModel.cpp
@@ -20,7 +20,7 @@ namespace UE::ConcertClientSharedSlate
 
 		static void DiscoverComponents(
 			const UObject& Object,
-			const TFunctionRef<EBreakBehavior(const FSoftObjectPath& Child, ConcertSharedSlate::EChildRelationship Relationship)>& Callback
+			const TFunctionRef<EBreakBehavior(const TSoftObjectPtr<>& Child, ConcertSharedSlate::EChildRelationship Relationship)>& Callback
 			)
 		{
 			using namespace ConcertSharedSlate;
@@ -64,15 +64,15 @@ namespace UE::ConcertClientSharedSlate
 	}
 
 	void FEditorObjectHierarchyModel::ForEachDirectChild(
-		const FSoftObjectPath& Parent,
-		TFunctionRef<EBreakBehavior(const FSoftObjectPath& Object, ConcertSharedSlate::EChildRelationship Relationship)> Callback,
+		const TSoftObjectPtr<>& Parent,
+		TFunctionRef<EBreakBehavior(const TSoftObjectPtr<>& Object, ConcertSharedSlate::EChildRelationship Relationship)> Callback,
 		ConcertSharedSlate::EChildRelationshipFlags InclusionFlags
 		) const
 	{
 		using namespace ConcertSharedSlate;
 
 		// In the context of replication, only resolve objects but do not load them. Case: client 1 may be viewing objects of client 2 but they are in different worlds.
-		const UObject* Object = Parent.ResolveObject();
+		const UObject* Object = Parent.Get();
 		if (!Object)
 		{
 			return;
@@ -95,12 +95,12 @@ namespace UE::ConcertClientSharedSlate
 		}
 	}
 
-	TOptional<ConcertSharedSlate::IObjectHierarchyModel::FParentInfo> FEditorObjectHierarchyModel::GetParentInfo(const FSoftObjectPath& ChildObject) const
+	TOptional<ConcertSharedSlate::IObjectHierarchyModel::FParentInfo> FEditorObjectHierarchyModel::GetParentInfo(const TSoftObjectPtr<>& ChildObject) const
 	{
 		using namespace ConcertSharedSlate;
 
 		// In the context of replication, only resolve objects but do not load them. Case: client 1 may be viewing objects of client 2 but they are in different worlds.
-		const UObject* Object = ChildObject.ResolveObject();
+		const UObject* Object = ChildObject.Get();
 		if (!Object)
 		{
 			return {};
diff --git a/Engine/Plugins/Developer/Concert/ConcertUI/ConcertClientSharedSlate/Source/ConcertClientSharedSlate/Private/Replication/Editor/Model/Object/EditorObjectHierarchyModel.h b/Engine/Plugins/Developer/Concert/ConcertUI/ConcertClientSharedSlate/Source/ConcertClientSharedSlate/Private/Replication/Editor/Model/Object/EditorObjectHierarchyModel.h
--- a/Engine/Plugins/Developer/Concert/ConcertUI/ConcertClientSharedSlate/Source/ConcertClientSharedSlate/Private/Replication/Editor/Model/Object/EditorObjectHierarchyModel.h
+++ b/Engine/Plugins/Developer/Concert/ConcertUI/ConcertClientSharedSlate/Source/ConcertClientSharedSlate/Private/Replication/Editor/Model/Object/EditorObjectHierarchyModel.h
@@ -23,11 +23,11 @@ namespace UE::ConcertClientSharedSlate
 
 		//~ Begin IObjectHierarchy Interface
 		virtual void ForEachDirectChild(
-			const FSoftObjectPath& Parent,
-			TFunctionRef<EBreakBehavior(const FSoftObjectPath& Object, ConcertSharedSlate::EChildRelationship Relationship)> Callback,
+			const TSoftObjectPtr<>& Parent,
+			TFunctionRef<EBreakBehavior(const TSoftObjectPtr<>& Object, ConcertSharedSlate::EChildRelationship Relationship)> Callback,
 			ConcertSharedSlate::EChildRelationshipFlags InclusionFlags
 			) const override;
-		virtual TOptional<FParentInfo> GetParentInfo(const FSoftObjectPath& ChildObject) const override;
+		virtual TOptional<FParentInfo> GetParentInfo(const TSoftObjectPtr<>& ChildObject) const override;
 		//~ End IObjectHierarchy Interface
 	};
 }
diff --git a/Engine/Plugins/Developer/Concert/ConcertUI/ConcertClientSharedSlate/Source/ConcertClientSharedSlate/Private/Replication/Editor/Model/PropertySource/SelectPropertyFromUClassModel.cpp b/Engine/Plugins/Developer/Concert/ConcertUI/ConcertClientSharedSlate/Source/ConcertClientSharedSlate/Private/Replication/Editor/Model/PropertySource/SelectPropertyFromUClassModel.cpp
--- a/Engine/Plugins/Developer/Concert/ConcertUI/ConcertClientSharedSlate/Source/ConcertClientSharedSlate/Private/Replication/Editor/Model/PropertySource/SelectPropertyFromUClassModel.cpp
+++ b/Engine/Plugins/Developer/Concert/ConcertUI/ConcertClientSharedSlate/Source/ConcertClientSharedSlate/Private/Replication/Editor/Model/PropertySource/SelectPropertyFromUClassModel.cpp
@@ -13,9 +13,9 @@ namespace UE::ConcertClientSharedSlate
 		: UClassIteratorSource(MakeShared<FConcertSyncCoreReplicatedPropertySource>())
 	{}
 
-	TSharedRef<ConcertSharedSlate::IPropertySourceModel> FSelectPropertyFromUClassModel::GetPropertySource(const FSoftClassPath& Class) const
+	TSharedRef<ConcertSharedSlate::IPropertySourceModel> FSelectPropertyFromUClassModel::GetPropertySource(const ConcertSharedSlate::FPropertySourceContext& Context) const
 	{
-		UClass* LoadedClass = Class.TryLoadClass<UObject>();
+		UClass* LoadedClass = Context.Class.TryLoadClass<UObject>();
 		if (LoadedClass)
 		{
 			UClassIteratorSource->SetClass(LoadedClass);
@@ -23,7 +23,7 @@ namespace UE::ConcertClientSharedSlate
 		else
 		{
 			UClassIteratorSource->SetClass(nullptr);
-			UE_LOG(LogConcert, Warning, TEXT("Could not resolve class %s. Properties will not be available."), *Class.ToString());
+			UE_LOG(LogConcert, Warning, TEXT("Could not resolve class %s. Properties will not be available."), *Context.Class.ToString());
 		}
 		return UClassIteratorSource;
 	}
diff --git a/Engine/Plugins/Developer/Concert/ConcertUI/ConcertClientSharedSlate/Source/ConcertClientSharedSlate/Public/Replication/Editor/Model/PropertySource/SelectPropertyFromUClassModel.h b/Engine/Plugins/Developer/Concert/ConcertUI/ConcertClientSharedSlate/Source/ConcertClientSharedSlate/Public/Replication/Editor/Model/PropertySource/SelectPropertyFromUClassModel.h
--- a/Engine/Plugins/Developer/Concert/ConcertUI/ConcertClientSharedSlate/Source/ConcertClientSharedSlate/Public/Replication/Editor/Model/PropertySource/SelectPropertyFromUClassModel.h
+++ b/Engine/Plugins/Developer/Concert/ConcertUI/ConcertClientSharedSlate/Source/ConcertClientSharedSlate/Public/Replication/Editor/Model/PropertySource/SelectPropertyFromUClassModel.h
@@ -23,7 +23,7 @@ namespace UE::ConcertClientSharedSlate
 		FSelectPropertyFromUClassModel();
 
 		//~ Begin IPropertySelectionSourceModel Interface
-		virtual TSharedRef<ConcertSharedSlate::IPropertySourceModel> GetPropertySource(const FSoftClassPath& Class) const override;
+		virtual TSharedRef<ConcertSharedSlate::IPropertySourceModel> GetPropertySource(const ConcertSharedSlate::FPropertySourceContext& Context) const override;
 		//~ End IPropertySelectionSourceModel Interface
 
 	private:
diff --git a/Engine/Plugins/Developer/Concert/ConcertUI/ConcertSharedSlate/Source/ConcertSharedSlate/Private/Replication/Editor/Model/Object/IObjectHierarchyModel.cpp b/Engine/Plugins/Developer/Concert/ConcertUI/ConcertSharedSlate/Source/ConcertSharedSlate/Private/Replication/Editor/Model/Object/IObjectHierarchyModel.cpp
--- a/Engine/Plugins/Developer/Concert/ConcertUI/ConcertSharedSlate/Source/ConcertSharedSlate/Private/Replication/Editor/Model/Object/IObjectHierarchyModel.cpp
+++ b/Engine/Plugins/Developer/Concert/ConcertUI/ConcertSharedSlate/Source/ConcertSharedSlate/Private/Replication/Editor/Model/Object/IObjectHierarchyModel.cpp
@@ -7,8 +7,8 @@
 namespace UE::ConcertSharedSlate
 {
 	void IObjectHierarchyModel::ForEachChildRecursive(
-		const FSoftObjectPath& Root,
-		TFunctionRef<EBreakBehavior(const FSoftObjectPath& Parent, const FSoftObjectPath& ChildObject, EChildRelationship Relationship)> Callback,
+		const TSoftObjectPtr<>& Root,
+		TFunctionRef<EBreakBehavior(const TSoftObjectPtr<>& Parent, const TSoftObjectPtr<>& ChildObject, EChildRelationship Relationship)> Callback,
 		EChildRelationshipFlags InclusionFlags
 		) const
 	{
@@ -16,10 +16,10 @@ namespace UE::ConcertSharedSlate
 		{
 			static EBreakBehavior Visit(
 				const IObjectHierarchyModel& Model,
-				const FSoftObjectPath& Parent,
-				const FSoftObjectPath& ChildObject,
+				const TSoftObjectPtr<>& Parent,
+				const TSoftObjectPtr<>& ChildObject,
 				EChildRelationship Relationship,
-				TFunctionRef<EBreakBehavior(const FSoftObjectPath&, const FSoftObjectPath&, EChildRelationship)> Callback,
+				TFunctionRef<EBreakBehavior(const TSoftObjectPtr<>&, const TSoftObjectPtr<>&, EChildRelationship)> Callback,
 				EChildRelationshipFlags InclusionFlags
 				)
 			{
@@ -29,7 +29,7 @@ namespace UE::ConcertSharedSlate
 				}
 				
 				EBreakBehavior Result = EBreakBehavior::Continue;
-				Model.ForEachDirectChild(ChildObject, [&Model, &ChildObject, &Callback, &InclusionFlags, &Result](const FSoftObjectPath& ChildOfChild, EChildRelationship Relationship)
+				Model.ForEachDirectChild(ChildObject, [&Model, &ChildObject, &Callback, &InclusionFlags, &Result](const TSoftObjectPtr<>& ChildOfChild, EChildRelationship Relationship)
 				{
 					Result = Visit(Model, ChildObject, ChildOfChild, Relationship, Callback, InclusionFlags);
 					return Result;
@@ -38,7 +38,7 @@ namespace UE::ConcertSharedSlate
 			}
 		};
 			
-		ForEachDirectChild(Root, [this, &Root, &Callback, &InclusionFlags](const FSoftObjectPath& ChildObject, EChildRelationship Relationship)
+		ForEachDirectChild(Root, [this, &Root, &Callback, &InclusionFlags](const TSoftObjectPtr<>& ChildObject, EChildRelationship Relationship)
 		{
 			return FHelper::Visit(*this, Root, ChildObject, Relationship, Callback, InclusionFlags);
 		}, InclusionFlags);
diff --git a/Engine/Plugins/Developer/Concert/ConcertUI/ConcertSharedSlate/Source/ConcertSharedSlate/Private/Replication/Editor/Utils/PropertyEnumerationUtils.cpp b/Engine/Plugins/Developer/Concert/ConcertUI/ConcertSharedSlate/Source/ConcertSharedSlate/Private/Replication/Editor/Utils/PropertyEnumerationUtils.cpp
--- a/Engine/Plugins/Developer/Concert/ConcertUI/ConcertSharedSlate/Source/ConcertSharedSlate/Private/Replication/Editor/Utils/PropertyEnumerationUtils.cpp
+++ b/Engine/Plugins/Developer/Concert/ConcertUI/ConcertSharedSlate/Source/ConcertSharedSlate/Private/Replication/Editor/Utils/PropertyEnumerationUtils.cpp
@@ -12,9 +12,9 @@
 namespace UE::ConcertSharedSlate
 {
 	/** Gets the class from the model or loads it. This function is designed to be used without assuming that it is run in editor-builds. */
-	FSoftClassPath GetObjectClassFromModelOrLoad(const FSoftObjectPath& Object, const IReplicationStreamModel& Model)
+	FSoftClassPath GetObjectClassFromModelOrLoad(const TSoftObjectPtr<>& Object, const IReplicationStreamModel& Model)
 	{
-		const FSoftClassPath ResolvedClass = Model.GetObjectClass(Object);
+		const FSoftClassPath ResolvedClass = Model.GetObjectClass(Object.GetUniqueID());
 		if (ResolvedClass.IsValid())
 		{
 			return ResolvedClass;
@@ -22,14 +22,14 @@ namespace UE::ConcertSharedSlate
 
 #if WITH_EDITOR
 		// The object is not yet in the model.
-		const UObject* LoadedObject = Object.ResolveObject();
+		const UObject* LoadedObject = Object.Get();
 		return LoadedObject ? LoadedObject->GetClass() : FSoftClassPath{};
 #else
 		return FSoftClassPath{};
 #endif
 	}
 
-	void EnumerateProperties(TConstArrayView<FSoftObjectPath> Objects, const IReplicationStreamModel& Model, const IPropertySelectionSourceModel* OptionalSource, FEnumerateProperties Callback)
+	void EnumerateProperties(TConstArrayView<TSoftObjectPtr<>> Objects, const IReplicationStreamModel& Model, const IPropertySelectionSourceModel* OptionalSource, FEnumerateProperties Callback)
 	{
 		if (OptionalSource)
 		{
@@ -42,14 +42,14 @@ namespace UE::ConcertSharedSlate
 	}
 
 	/** Enumerates the properties that are assigned to the object in Model */
-	void EnumerateRegisteredPropertiesOnly(TConstArrayView<FSoftObjectPath> Objects, const IReplicationStreamModel& Model, FEnumerateProperties Callback)
+	void EnumerateRegisteredPropertiesOnly(TConstArrayView<TSoftObjectPtr<>> Objects, const IReplicationStreamModel& Model, FEnumerateProperties Callback)
 	{
-		for (const FSoftObjectPath& Object : Objects)
+		for (const TSoftObjectPtr<>& Object : Objects)
 		{
 			const FSoftClassPath ObjectClass = GetObjectClassFromModelOrLoad(Object, Model);
 
 			EBreakBehavior BreakBehavior = EBreakBehavior::Continue;
-			Model.ForEachProperty(Object, [&Callback, &ObjectClass, &BreakBehavior](const FConcertPropertyChain& Chain)
+			Model.ForEachProperty(Object.GetUniqueID(), [&Callback, &ObjectClass, &BreakBehavior](const FConcertPropertyChain& Chain)
 			{
 				BreakBehavior = Callback(ObjectClass, Chain);
 				return BreakBehavior;
@@ -63,21 +63,17 @@ namespace UE::ConcertSharedSlate
 	}
 
 	/** Enumerate the properties that are selectable in Source (e.g. all properties in that class, @see FSelectPropertyFromUClassModel). */
-	void EnumerateAllProperties(TConstArrayView<FSoftObjectPath> Objects, const IPropertySelectionSourceModel& Source, const IReplicationStreamModel& Model, FEnumerateProperties Callback)
+	void EnumerateAllProperties(TConstArrayView<TSoftObjectPtr<>> Objects, const IPropertySelectionSourceModel& Source, const IReplicationStreamModel& Model, FEnumerateProperties Callback)
 	{
 		TSet<FSoftClassPath> VisitedClasses; 
 			
-		for (const FSoftObjectPath& Object : Objects)
+		for (const TSoftObjectPtr<>& Object : Objects)
 		{
 			const FSoftClassPath ObjectClass = GetObjectClassFromModelOrLoad(Object, Model);
-			if (VisitedClasses.Contains(ObjectClass))
-			{
-				continue;
-			}
-			VisitedClasses.Add(ObjectClass);
+			const FPropertySourceContext ObjectQueryContext(Object, ObjectClass);
 					
 			EBreakBehavior BreakBehavior = EBreakBehavior::Continue;
-			Source.GetPropertySource(ObjectClass)->EnumerateSelectableItems([&Callback, &ObjectClass, &BreakBehavior](const FSelectablePropertyInfo& PropertyInfo)
+			Source.GetPropertySource(ObjectQueryContext)->EnumerateSelectableItems([&Callback, &ObjectClass, &BreakBehavior](const FSelectablePropertyInfo& PropertyInfo)
 			{
 				BreakBehavior = Callback(ObjectClass, PropertyInfo.Property);
 				return BreakBehavior;
diff --git a/Engine/Plugins/Developer/Concert/ConcertUI/ConcertSharedSlate/Source/ConcertSharedSlate/Private/Replication/Editor/Utils/PropertyEnumerationUtils.h b/Engine/Plugins/Developer/Concert/ConcertUI/ConcertSharedSlate/Source/ConcertSharedSlate/Private/Replication/Editor/Utils/PropertyEnumerationUtils.h
--- a/Engine/Plugins/Developer/Concert/ConcertUI/ConcertSharedSlate/Source/ConcertSharedSlate/Private/Replication/Editor/Utils/PropertyEnumerationUtils.h
+++ b/Engine/Plugins/Developer/Concert/ConcertUI/ConcertSharedSlate/Source/ConcertSharedSlate/Private/Replication/Editor/Utils/PropertyEnumerationUtils.h
@@ -2,10 +2,10 @@
 
 #pragma once
 
-
 #include "Containers/ContainersFwd.h"
 #include "HAL/Platform.h"
 #include "Templates/FunctionFwd.h"
+#include "UObject/SoftObjectPtr.h"
 
 enum class EBreakBehavior : uint8;
 struct FConcertPropertyChain;
@@ -19,12 +19,12 @@ namespace UE::ConcertSharedSlate
 	using FEnumerateProperties = TFunctionRef<EBreakBehavior(const FSoftClassPath&, const FConcertPropertyChain&)>;
 	
 	/** Gets the class from the model or loads it. This function is designed to be used without assuming that it is run in editor-builds. */
-	FSoftClassPath GetObjectClassFromModelOrLoad(const FSoftObjectPath& Object, const IReplicationStreamModel& Model);
+	FSoftClassPath GetObjectClassFromModelOrLoad(const TSoftObjectPtr<>& Object, const IReplicationStreamModel& Model);
 
 	/** Calls EnumerateRegisteredPropertiesOnly or EnumerateAllProperties depending on whether OptionalSource is nullptr. */
-	void EnumerateProperties(TConstArrayView<FSoftObjectPath> Objects, const IReplicationStreamModel& Model, const IPropertySelectionSourceModel* OptionalSource, FEnumerateProperties Callback);
+	void EnumerateProperties(TConstArrayView<TSoftObjectPtr<>> Objects, const IReplicationStreamModel& Model, const IPropertySelectionSourceModel* OptionalSource, FEnumerateProperties Callback);
 	/** Enumerates the properties that are assigned to the object in Model */
-	void EnumerateRegisteredPropertiesOnly(TConstArrayView<FSoftObjectPath> Objects, const IReplicationStreamModel& Model, FEnumerateProperties Callback);
+	void EnumerateRegisteredPropertiesOnly(TConstArrayView<TSoftObjectPtr<>> Objects, const IReplicationStreamModel& Model, FEnumerateProperties Callback);
 	/** Enumerate the properties that are selectable in Source (e.g. all properties in that class, @see FSelectPropertyFromUClassModel). */
-	void EnumerateAllProperties(TConstArrayView<FSoftObjectPath> Objects, const IPropertySelectionSourceModel& Source, const IReplicationStreamModel& Model, FEnumerateProperties Callback);
+	void EnumerateAllProperties(TConstArrayView<TSoftObjectPtr<>> Objects, const IPropertySelectionSourceModel& Source, const IReplicationStreamModel& Model, FEnumerateProperties Callback);
 }
\ No newline at end of file
diff --git a/Engine/Plugins/Developer/Concert/ConcertUI/ConcertSharedSlate/Source/ConcertSharedSlate/Private/Replication/Editor/View/IReplicationStreamViewer.cpp b/Engine/Plugins/Developer/Concert/ConcertUI/ConcertSharedSlate/Source/ConcertSharedSlate/Private/Replication/Editor/View/IReplicationStreamViewer.cpp
new file mode 100644
--- /dev/null
+++ b/Engine/Plugins/Developer/Concert/ConcertUI/ConcertSharedSlate/Source/ConcertSharedSlate/Private/Replication/Editor/View/IReplicationStreamViewer.cpp
@@ -0,0 +1,13 @@
+﻿// Copyright Epic Games, Inc. All Rights Reserved.
+
+#include "Replication/Editor/View/IReplicationStreamViewer.h"
+
+namespace UE::ConcertSharedSlate
+{
+	TArray<FSoftObjectPath> IReplicationStreamViewer::GetObjectsBeingPropertyEdited() const
+	{
+		TArray<FSoftObjectPath> Result;
+		Algo::Transform(GetSelectedObjects(), Result, [](const TSoftObjectPtr<>& Object){ return Object.GetUniqueID(); });
+		return Result;
+	}
+}
diff --git a/Engine/Plugins/Developer/Concert/ConcertUI/ConcertSharedSlate/Source/ConcertSharedSlate/Private/Replication/Editor/View/ObjectEditor/SBaseReplicationStreamEditor.cpp b/Engine/Plugins/Developer/Concert/ConcertUI/ConcertSharedSlate/Source/ConcertSharedSlate/Private/Replication/Editor/View/ObjectEditor/SBaseReplicationStreamEditor.cpp
--- a/Engine/Plugins/Developer/Concert/ConcertUI/ConcertSharedSlate/Source/ConcertSharedSlate/Private/Replication/Editor/View/ObjectEditor/SBaseReplicationStreamEditor.cpp
+++ b/Engine/Plugins/Developer/Concert/ConcertUI/ConcertSharedSlate/Source/ConcertSharedSlate/Private/Replication/Editor/View/ObjectEditor/SBaseReplicationStreamEditor.cpp
@@ -93,9 +93,9 @@ namespace UE::ConcertSharedSlate
 		ReplicationViewer->RequestPropertyColumnResort(ColumnId);
 	}
 
-	TArray<FSoftObjectPath> SBaseReplicationStreamEditor::GetObjectsBeingPropertyEdited() const
+	TArray<TSoftObjectPtr<>> SBaseReplicationStreamEditor::GetSelectedObjects() const
 	{
-		return ReplicationViewer->GetObjectsBeingPropertyEdited();
+		return ReplicationViewer->GetSelectedObjects();
 	}
 
 	FHoverRowContent SBaseReplicationStreamEditor::MakeHoveredRowContent(const TSharedPtr<FReplicatedObjectData>& Data) const
@@ -145,7 +145,7 @@ namespace UE::ConcertSharedSlate
 		// Newly added objects should be automatically selected
 		if (!AddedObjects.IsEmpty())
 		{
-			TArray<FSoftObjectPath> TopLevelObjects;
+			TArray<TSoftObjectPtr<>> TopLevelObjects;
 			
 			// Goal: select an object so the property view immediately shows properties for some object.
 			// Problem: if objects have different classes, property view will be empty (incompatible class)
@@ -155,7 +155,7 @@ namespace UE::ConcertSharedSlate
 				[this, &AddedObjects](const UObject* Object)
 				{
 					const TOptional<IObjectHierarchyModel::FParentInfo> ParentInfo = ObjectHierarchy->GetParentInfo(Object);
-					const bool bIsTopOfHierarchy = !ParentInfo || !AddedObjects.ContainsByPredicate([&ParentInfo](UObject* AddedObject){ return FSoftObjectPath(AddedObject) == ParentInfo->Parent; });
+					const bool bIsTopOfHierarchy = !ParentInfo || !AddedObjects.ContainsByPredicate([&ParentInfo](UObject* AddedObject){ return AddedObject == ParentInfo->Parent; });
 					return bIsTopOfHierarchy;
 				},
 				[](const UObject* Object){ return Object; });
@@ -249,8 +249,8 @@ namespace UE::ConcertSharedSlate
 			EUserInterfaceActionType::Button
 		);
 
-		TArray<FSoftObjectPath> SelectedObjects;
-		Algo::Transform(ReplicationViewer->GetSelectedOutlinerObjects(), SelectedObjects, [](const TSharedPtr<FReplicatedObjectData>& Data){ return Data->GetObjectPath(); });
+		TArray<TSoftObjectPtr<>> SelectedObjects;
+		Algo::Transform(ReplicationViewer->GetSelectedOutlinerObjects(), SelectedObjects, [](const TSharedPtr<FReplicatedObjectData>& Data){ return Data->GetObjectPtr(); });
 		OnExtendObjectsContextMenuDelegate.ExecuteIfBound(MenuBuilder, SelectedObjects);
 		
 		return MenuBuilder.MakeWidget();
diff --git a/Engine/Plugins/Developer/Concert/ConcertUI/ConcertSharedSlate/Source/ConcertSharedSlate/Private/Replication/Editor/View/ObjectEditor/SBaseReplicationStreamEditor.h b/Engine/Plugins/Developer/Concert/ConcertUI/ConcertSharedSlate/Source/ConcertSharedSlate/Private/Replication/Editor/View/ObjectEditor/SBaseReplicationStreamEditor.h
--- a/Engine/Plugins/Developer/Concert/ConcertUI/ConcertSharedSlate/Source/ConcertSharedSlate/Private/Replication/Editor/View/ObjectEditor/SBaseReplicationStreamEditor.h
+++ b/Engine/Plugins/Developer/Concert/ConcertUI/ConcertSharedSlate/Source/ConcertSharedSlate/Private/Replication/Editor/View/ObjectEditor/SBaseReplicationStreamEditor.h
@@ -24,7 +24,6 @@ namespace UE::ConcertSharedSlate
 	class IObjectSelectionSourceModel;
 	class IPropertyAssignmentView;
 	class IPropertySelectionSourceModel;
-	class IReplicationSubobjectView;
 	class IObjectHierarchyModel;
 	class SReplicationStreamViewer;
 	
@@ -97,7 +96,7 @@ namespace UE::ConcertSharedSlate
 		virtual void Refresh() override;
 		virtual void RequestObjectColumnResort(const FName& ColumnId) override;
 		virtual void RequestPropertyColumnResort(const FName& ColumnId) override;
-		virtual TArray<FSoftObjectPath> GetObjectsBeingPropertyEdited() const override;
+		virtual TArray<TSoftObjectPtr<>> GetSelectedObjects() const override;
 		//~ End IReplicationStreamEditor Interface
 
 	private:
diff --git a/Engine/Plugins/Developer/Concert/ConcertUI/ConcertSharedSlate/Source/ConcertSharedSlate/Private/Replication/Editor/View/ObjectViewer/SReplicatedPropertyView.cpp b/Engine/Plugins/Developer/Concert/ConcertUI/ConcertSharedSlate/Source/ConcertSharedSlate/Private/Replication/Editor/View/ObjectViewer/SReplicatedPropertyView.cpp
--- a/Engine/Plugins/Developer/Concert/ConcertUI/ConcertSharedSlate/Source/ConcertSharedSlate/Private/Replication/Editor/View/ObjectViewer/SReplicatedPropertyView.cpp
+++ b/Engine/Plugins/Developer/Concert/ConcertUI/ConcertSharedSlate/Source/ConcertSharedSlate/Private/Replication/Editor/View/ObjectViewer/SReplicatedPropertyView.cpp
@@ -32,7 +32,7 @@ namespace UE::ConcertSharedSlate
 	void SReplicatedPropertyView::RefreshPropertyData()
 	{
 		SCOPED_CONCERT_TRACE(RefreshPropertyData);
-		TArray<FSoftObjectPath> SelectedObjects = GetObjectsSelectedForPropertyEditing();
+		TArray<TSoftObjectPtr<>> SelectedObjects = GetObjectsSelectedForPropertyEditing();
 		if (SelectedObjects.IsEmpty())
 		{
 			SetPropertyContent(EReplicatedPropertyContent::NoSelection);
@@ -51,12 +51,12 @@ namespace UE::ConcertSharedSlate
 		SetPropertyContent(EReplicatedPropertyContent::Properties);
 	}
 
-	TArray<FSoftObjectPath> SReplicatedPropertyView::GetObjectsSelectedForPropertyEditing() const
+	TArray<TSoftObjectPtr<>> SReplicatedPropertyView::GetObjectsSelectedForPropertyEditing() const
 	{
-		TArray<FSoftObjectPath> Result;
+		TArray<TSoftObjectPtr<>> Result;
 		Algo::Transform(GetSelectedRootObjectsDelegate.Execute(), Result, [](const TSharedPtr<FReplicatedObjectData>& ObjectData)
 		{
-			return ObjectData->GetObjectPath();
+			return ObjectData->GetObjectPtr();
 		});
 		return Result;
 	}
@@ -92,12 +92,12 @@ namespace UE::ConcertSharedSlate
 			];
 	}
 	
-	TOptional<FSoftClassPath> SReplicatedPropertyView::GetClassForPropertiesFromSelection(const TArray<FSoftObjectPath>& Objects) const
+	TOptional<FSoftClassPath> SReplicatedPropertyView::GetClassForPropertiesFromSelection(const TArray<TSoftObjectPtr<>>& Objects) const
 	{
 		FSoftClassPath SharedClass;
-		const bool bAllHaveSameClass = Algo::AllOf(Objects, [this, &SharedClass](const FSoftObjectPath& Object)
+		const bool bAllHaveSameClass = Algo::AllOf(Objects, [this, &SharedClass](const TSoftObjectPtr<>& Object)
 		{
-			const FSoftClassPath ObjectClass = GetObjectClass(TSoftObjectPtr<>{ Object });
+			const FSoftClassPath ObjectClass = GetObjectClass(Object);
 			SharedClass = SharedClass.IsValid() ? SharedClass : ObjectClass;
 			return ObjectClass == SharedClass;
 		});
diff --git a/Engine/Plugins/Developer/Concert/ConcertUI/ConcertSharedSlate/Source/ConcertSharedSlate/Private/Replication/Editor/View/ObjectViewer/SReplicatedPropertyView.h b/Engine/Plugins/Developer/Concert/ConcertUI/ConcertSharedSlate/Source/ConcertSharedSlate/Private/Replication/Editor/View/ObjectViewer/SReplicatedPropertyView.h
--- a/Engine/Plugins/Developer/Concert/ConcertUI/ConcertSharedSlate/Source/ConcertSharedSlate/Private/Replication/Editor/View/ObjectViewer/SReplicatedPropertyView.h
+++ b/Engine/Plugins/Developer/Concert/ConcertUI/ConcertSharedSlate/Source/ConcertSharedSlate/Private/Replication/Editor/View/ObjectViewer/SReplicatedPropertyView.h
@@ -42,7 +42,7 @@ namespace UE::ConcertSharedSlate
 		/** Requests that the given column be resorted, if it currently affects the row sorting. */
 		void RequestResortForColumn(const FName& ColumnId) const { PropertyAssignmentView->RequestResortForColumn(ColumnId); }
 		
-		TArray<FSoftObjectPath> GetObjectsSelectedForPropertyEditing() const;
+		TArray<TSoftObjectPtr<>> GetObjectsSelectedForPropertyEditing() const;
 
 	private:
 		
@@ -72,7 +72,7 @@ namespace UE::ConcertSharedSlate
 		TSharedRef<SWidget> CreatePropertiesView(const FArguments& InArgs);
 		
 		/** Given the selected objects, determines whether they all have the same class and returns it if so. */
-		TOptional<FSoftClassPath> GetClassForPropertiesFromSelection(const TArray<FSoftObjectPath>& Objects) const;
+		TOptional<FSoftClassPath> GetClassForPropertiesFromSelection(const TArray<TSoftObjectPtr<>>& Objects) const;
 		/** Sets how to display this widget */
 		void SetPropertyContent(EReplicatedPropertyContent Content) const;
 
diff --git a/Engine/Plugins/Developer/Concert/ConcertUI/ConcertSharedSlate/Source/ConcertSharedSlate/Private/Replication/Editor/View/ObjectViewer/SReplicationStreamViewer.cpp b/Engine/Plugins/Developer/Concert/ConcertUI/ConcertSharedSlate/Source/ConcertSharedSlate/Private/Replication/Editor/View/ObjectViewer/SReplicationStreamViewer.cpp
--- a/Engine/Plugins/Developer/Concert/ConcertUI/ConcertSharedSlate/Source/ConcertSharedSlate/Private/Replication/Editor/View/ObjectViewer/SReplicationStreamViewer.cpp
+++ b/Engine/Plugins/Developer/Concert/ConcertUI/ConcertSharedSlate/Source/ConcertSharedSlate/Private/Replication/Editor/View/ObjectViewer/SReplicationStreamViewer.cpp
@@ -26,6 +26,19 @@
 
 namespace UE::ConcertSharedSlate
 {
+	namespace Private
+	{
+		static TSharedRef<FReplicatedObjectData> AllocateObjectData(FSoftObjectPath ObjectPath)
+		{
+			return MakeShared<FReplicatedObjectData>(MoveTemp(ObjectPath));
+		}
+
+		static TSharedRef<FReplicatedObjectData> AllocateObjectData(TSoftObjectPtr<> Object)
+		{
+			return MakeShared<FReplicatedObjectData>(MoveTemp(Object));
+		}
+	}
+	
 	void SReplicationStreamViewer::Construct(const FArguments& InArgs, const TSharedRef<IReplicationStreamModel>& InPropertiesModel)
 	{
 		PropertiesModel = InPropertiesModel;
@@ -58,12 +71,12 @@ namespace UE::ConcertSharedSlate
 		PropertySection->RequestResortForColumn(ColumnId);
 	}
 
-	TArray<FSoftObjectPath> SReplicationStreamViewer::GetObjectsBeingPropertyEdited() const
+	TArray<TSoftObjectPtr<>> SReplicationStreamViewer::GetSelectedObjects() const
 	{
 		return PropertySection->GetObjectsSelectedForPropertyEditing();
 	}
 
-	void SReplicationStreamViewer::SelectObjects(TConstArrayView<FSoftObjectPath> Objects, bool bAtEndOfTick)
+	void SReplicationStreamViewer::SelectObjects(TConstArrayView<TSoftObjectPtr<>> Objects, bool bAtEndOfTick)
 	{
 		SCOPED_CONCERT_TRACE(SelectObjects);
 		
@@ -76,7 +89,7 @@ namespace UE::ConcertSharedSlate
 		TArray<TSharedPtr<FReplicatedObjectData>> NewSelectedItems; 
 		Algo::TransformIf(AllObjectRowData, NewSelectedItems, [&Objects](const TSharedPtr<FReplicatedObjectData>& ObjectData)
 			{
-				return Objects.Contains(ObjectData->GetObjectPath());
+				return Objects.Contains(ObjectData->GetObjectPtr());
 			},
 			[](const TSharedPtr<FReplicatedObjectData>& ObjectData){ return ObjectData; }
 		);
@@ -86,7 +99,7 @@ namespace UE::ConcertSharedSlate
 		}
 	}
 
-	void SReplicationStreamViewer::ExpandObjects(TConstArrayView<FSoftObjectPath> Objects, bool bRecursive, bool bAtEndOfTick)
+	void SReplicationStreamViewer::ExpandObjects(TConstArrayView<TSoftObjectPtr<>> Objects, bool bRecursive, bool bAtEndOfTick)
 	{
 		SCOPED_CONCERT_TRACE(ExpandObjects);
 		
@@ -104,18 +117,18 @@ namespace UE::ConcertSharedSlate
 		
 		TArray<TSharedPtr<FReplicatedObjectData>> ItemsToExpand;
 		ItemsToExpand.Reserve(Objects.Num());
-		for (const FSoftObjectPath& Path : Objects)
+		for (const TSoftObjectPtr<>& Path : Objects)
 		{
-			if (const TSharedPtr<FReplicatedObjectData>* Item = PathToObjectDataCache.Find(Path))
+			if (const TSharedPtr<FReplicatedObjectData>* Item = PathToObjectDataCache.Find(Path.GetUniqueID()))
 			{
 				ItemsToExpand.Add(*Item);
 			}
 			
 			if (bRecursive && ObjectHierarchy)
 			{
-				ObjectHierarchy->ForEachChildRecursive(Path, [this, &ItemsToExpand](const FSoftObjectPath&, const FSoftObjectPath& ChildObject, EChildRelationship)
+				ObjectHierarchy->ForEachChildRecursive(Path, [this, &ItemsToExpand](const TSoftObjectPtr<>&, const TSoftObjectPtr<>& ChildObject, EChildRelationship)
 				{
-					if (const TSharedPtr<FReplicatedObjectData>* Item = PathToObjectDataCache.Find(ChildObject))
+					if (const TSharedPtr<FReplicatedObjectData>* Item = PathToObjectDataCache.Find(ChildObject.GetUniqueID()))
 					{
 						ItemsToExpand.Add(*Item);
 					}
@@ -188,11 +201,6 @@ namespace UE::ConcertSharedSlate
 		IReplicationStreamViewer::Tick(AllottedGeometry, InCurrentTime, InDeltaTime);
 	}
 
-	TSharedRef<FReplicatedObjectData> SReplicationStreamViewer::AllocateObjectData(FSoftObjectPath ObjectPath)
-	{
-		return MakeShared<FReplicatedObjectData>(MoveTemp(ObjectPath));
-	}
-
 	TSharedRef<SWidget> SReplicationStreamViewer::CreateContentWidget(const FArguments& InArgs)
 	{
 		return SNew(SSplitter)
@@ -317,13 +325,15 @@ namespace UE::ConcertSharedSlate
 		// An alternative would be to change RefreshObjectData to be called with two variables ObjectsAdded and ObjectsRemoved.
 		IterateDisplayableObjects([this, &NewPathToObjectDataCache](const FSoftObjectPath& ObjectPath) mutable
 		{
-			TOptional<IObjectHierarchyModel::FParentInfo> ParentInfo = ObjectHierarchy->GetParentInfo(ObjectPath);
+			TOptional<IObjectHierarchyModel::FParentInfo> ParentInfo = ObjectHierarchy->GetParentInfo(TSoftObjectPtr(ObjectPath));
 			const bool bIsActor = !ParentInfo; 
 			if (bIsActor || ShouldDisplayChildObject(ObjectPath, ParentInfo->Relationship))
 			{
 				const TSharedPtr<FReplicatedObjectData>* ExistingItem = PathToObjectDataCache.Find(ObjectPath);
 				ExistingItem = ExistingItem ? ExistingItem : NewPathToObjectDataCache.Find(ObjectPath);
-				const TSharedRef<FReplicatedObjectData> Item = ExistingItem ? ExistingItem->ToSharedRef() : AllocateObjectData(ObjectPath);
+				const TSharedRef<FReplicatedObjectData> Item = ExistingItem
+					? ExistingItem->ToSharedRef()
+					: Private::AllocateObjectData(ObjectPath);
 				AllObjectRowData.AddUnique(Item);
 				NewPathToObjectDataCache.Emplace(ObjectPath, Item);
 				
@@ -450,19 +460,23 @@ namespace UE::ConcertSharedSlate
 		}
 		
 		// Add all objects that appear in the hierarchy of ReplicatedObjectData
-		const auto AddItem = [this, &NewPathToObjectDataCache](const FSoftObjectPath& ObjectPath)
+		const auto AddItem = [this, &NewPathToObjectDataCache](const TSoftObjectPtr<>& Object)
 		{
+			const FSoftObjectPath& ObjectPath = Object.GetUniqueID();
 			const TSharedPtr<FReplicatedObjectData>* ExistingItem = PathToObjectDataCache.Find(ObjectPath);
 			ExistingItem = ExistingItem ? ExistingItem : NewPathToObjectDataCache.Find(ObjectPath);
-			const TSharedRef<FReplicatedObjectData> Item = ExistingItem ? ExistingItem->ToSharedRef() : AllocateObjectData(ObjectPath);
+			const TSharedRef<FReplicatedObjectData> Item = ExistingItem
+				? ExistingItem->ToSharedRef()
+				: Private::AllocateObjectData(ObjectPath);
 			AllObjectRowData.AddUnique(Item);
 			NewPathToObjectDataCache.Emplace(ObjectPath, Item);
 		};
 		
-		AddItem(OwningActor);
-		ObjectHierarchy->ForEachChildRecursive(OwningActor, [this, &AddItem](const FSoftObjectPath&, const FSoftObjectPath& ChildObject, EChildRelationship Relationship)
+		const TSoftObjectPtr ActorPtr(OwningActor);
+		AddItem(ActorPtr);
+		ObjectHierarchy->ForEachChildRecursive(ActorPtr, [this, &AddItem](const TSoftObjectPtr<>&, const TSoftObjectPtr<>& ChildObject, EChildRelationship Relationship)
 		{
-			if (ShouldDisplayChildObject(ChildObject, Relationship))
+			if (ShouldDisplayChildObject(ChildObject.GetUniqueID(), Relationship))
 			{
 				AddItem(ChildObject);
 			}
@@ -474,15 +488,15 @@ namespace UE::ConcertSharedSlate
 	{
 		// Important: this view should be possible to be built in programs, so it should not reference things like AActor, UActorComponent, ResolveObject, etc. directly.
 		
-		const FSoftObjectPath& SearchedObject = ReplicatedObjectData->GetObjectPath();
+		const TSoftObjectPtr<>& SearchedObject = ReplicatedObjectData->GetObjectPtr();
 		if (!ObjectHierarchy)
 		{
 			return;
 		}
 
-		ObjectHierarchy->ForEachDirectChild(SearchedObject, [this, &ProcessChild](const FSoftObjectPath& ChildObject, EChildRelationship Relationship)
+		ObjectHierarchy->ForEachDirectChild(SearchedObject, [this, &ProcessChild](const TSoftObjectPtr<>& ChildObject, EChildRelationship Relationship)
 		{
-			if (const TSharedPtr<FReplicatedObjectData>* ObjectData = PathToObjectDataCache.Find(ChildObject))
+			if (const TSharedPtr<FReplicatedObjectData>* ObjectData = PathToObjectDataCache.Find(ChildObject.GetUniqueID()))
 			{
 				ProcessChild(*ObjectData);
 			}
diff --git a/Engine/Plugins/Developer/Concert/ConcertUI/ConcertSharedSlate/Source/ConcertSharedSlate/Private/Replication/Editor/View/ObjectViewer/SReplicationStreamViewer.h b/Engine/Plugins/Developer/Concert/ConcertUI/ConcertSharedSlate/Source/ConcertSharedSlate/Private/Replication/Editor/View/ObjectViewer/SReplicationStreamViewer.h
--- a/Engine/Plugins/Developer/Concert/ConcertUI/ConcertSharedSlate/Source/ConcertSharedSlate/Private/Replication/Editor/View/ObjectViewer/SReplicationStreamViewer.h
+++ b/Engine/Plugins/Developer/Concert/ConcertUI/ConcertSharedSlate/Source/ConcertSharedSlate/Private/Replication/Editor/View/ObjectViewer/SReplicationStreamViewer.h
@@ -95,16 +95,16 @@ namespace UE::ConcertSharedSlate
 		virtual void Refresh() override;
 		virtual void RequestObjectColumnResort(const FName& ColumnId) override;
 		virtual void RequestPropertyColumnResort(const FName& ColumnId) override;
-		virtual TArray<FSoftObjectPath> GetObjectsBeingPropertyEdited() const override;
+		virtual TArray<TSoftObjectPtr<>> GetSelectedObjects() const override;
 		//~ End IReplicationStreamViewer Interface
 
 		void RequestObjectDataRefresh() { bHasRequestedObjectRefresh = true; }
 		void RequestPropertyDataRefresh() { bHasRequestedPropertyRefresh = true; }
 
 		/** Selects the given objects. */
-		void SelectObjects(TConstArrayView<FSoftObjectPath> Objects, bool bAtEndOfTick = false);
+		void SelectObjects(TConstArrayView<TSoftObjectPtr<>> Objects, bool bAtEndOfTick = false);
 		/** Expands the given objects, recursively if desired. */
-		void ExpandObjects(TConstArrayView<FSoftObjectPath> Objects, bool bRecursive, bool bAtEndOfTick = false);
+		void ExpandObjects(TConstArrayView<TSoftObjectPtr<>> Objects, bool bRecursive, bool bAtEndOfTick = false);
 
 		/** @return Gets the root objects selected in the outliner; the subobject view chooses which of these objects (or their subobjects) end up in GetSelectedObjectShowingProperties. */
 		TArray<TSharedPtr<FReplicatedObjectData>> GetSelectedOutlinerObjects() const;
@@ -143,15 +143,13 @@ namespace UE::ConcertSharedSlate
 
 		bool bHasRequestedObjectRefresh = false;
 		bool bHasRequestedPropertyRefresh = false;
-		TArray<FSoftObjectPath> PendingToSelect;
-		TArray<FSoftObjectPath> PendingToExpand;
+		TArray<TSoftObjectPtr<>> PendingToSelect;
+		TArray<TSoftObjectPtr<>> PendingToExpand;
 		bool bPendingExpandRecursively = false;
 		
 		/** Optional. Whether a given object should be displayed. If this returns false on an object, none of its children will be shown either. */
 		FShouldDisplayObject ShouldDisplayObjectDelegate;
 
-		static TSharedRef<FReplicatedObjectData> AllocateObjectData(FSoftObjectPath ObjectPath);
-
 		// Widget creation helpers
 		TSharedRef<SWidget> CreateContentWidget(const FArguments& InArgs);
 		TSharedRef<SWidget> CreateOutlinerSection(const FArguments& InArgs);
diff --git a/Engine/Plugins/Developer/Concert/ConcertUI/ConcertSharedSlate/Source/ConcertSharedSlate/Private/Replication/Editor/View/Property/SMultiObjectAssignment.cpp b/Engine/Plugins/Developer/Concert/ConcertUI/ConcertSharedSlate/Source/ConcertSharedSlate/Private/Replication/Editor/View/Property/SMultiObjectAssignment.cpp
--- a/Engine/Plugins/Developer/Concert/ConcertUI/ConcertSharedSlate/Source/ConcertSharedSlate/Private/Replication/Editor/View/Property/SMultiObjectAssignment.cpp
+++ b/Engine/Plugins/Developer/Concert/ConcertUI/ConcertSharedSlate/Source/ConcertSharedSlate/Private/Replication/Editor/View/Property/SMultiObjectAssignment.cpp
@@ -11,15 +11,15 @@ namespace UE::ConcertSharedSlate
 {
 	namespace Private
 	{
-		static TMap<FSoftObjectPath, TArray<FSoftObjectPath>> BuildHierarchy(const TArray<FSoftObjectPath>& RootObjects, IObjectHierarchyModel& HierarchyModel, bool bShouldShowSubobjects)
+		static TMap<FSoftObjectPath, TArray<FSoftObjectPath>> BuildHierarchy(const TArray<TSoftObjectPtr<>>& RootObjects, IObjectHierarchyModel& HierarchyModel, bool bShouldShowSubobjects)
 		{
 			TMap<FSoftObjectPath, TArray<FSoftObjectPath>> Result;
 
-			for (const FSoftObjectPath& Object : RootObjects)
+			for (const TSoftObjectPtr<>& Object : RootObjects)
 			{
 				HierarchyModel.ForEachChildRecursive(
 					Object,
-					[&Result](const FSoftObjectPath& Parent, const FSoftObjectPath& ChildObject, EChildRelationship Relationship)
+					[&Result](const TSoftObjectPtr<>& Parent, const TSoftObjectPtr<>& ChildObject, EChildRelationship Relationship)
 					{
 						// ForEachChildRecursive iterates through the hierarchy as it would be displayed by the component panel in the details panel
 						// E.g. for ACharacter this would say
@@ -34,7 +34,7 @@ namespace UE::ConcertSharedSlate
 						ParentPathString.LeftInline(Index);
 
 						const FSoftObjectPath ParentPath(*ParentPathString);
-						Result.FindOrAdd(ParentPath).Add(ChildObject);
+						Result.FindOrAdd(ParentPath).Add(ChildObject.GetUniqueID());
 						return EBreakBehavior::Continue;
 					},
 					bShouldShowSubobjects ? EChildRelationshipFlags::All : EChildRelationshipFlags::Component
@@ -45,9 +45,9 @@ namespace UE::ConcertSharedSlate
 		}
 
 		static void ForEachCategory(
-			const TArray<FSoftObjectPath>& Start,
+			const TArray<TSoftObjectPtr<>>& Start,
 			const TMap<FSoftObjectPath, TArray<FSoftObjectPath>>& Hierarchy,
-			TFunctionRef<void(const TArray<FSoftObjectPath>& ContextObjects)> Callback
+			TFunctionRef<void(const TArray<TSoftObjectPtr<>>& ContextObjects)> Callback
 			)
 		{
 			if (Start.IsEmpty())
@@ -56,8 +56,8 @@ namespace UE::ConcertSharedSlate
 			}
 
 			// While multi-editing, if any of the subobjects do not match skip the entire category.
-			const FSoftObjectPath FirstEntry = *Start.CreateConstIterator();
-			const TArray<FSoftObjectPath>* FirstEntryChildren = Hierarchy.Find(FirstEntry);
+			const TSoftObjectPtr<> FirstEntry = Start[0];
+			const TArray<FSoftObjectPath>* FirstEntryChildren = Hierarchy.Find(FirstEntry.GetUniqueID());
 			if (!FirstEntryChildren)
 			{
 				return;
@@ -72,11 +72,11 @@ namespace UE::ConcertSharedSlate
 				// +1 for the "." preceding the child, otherwise we'd get ".Component"
 				ChildName.RightChopInline(Base.Len() + 1);
 
-				TArray<FSoftObjectPath> RelatedChildren;
+				TArray<TSoftObjectPtr<>> RelatedChildren;
 				// This intentionally iterates FirstEntry as well
-				for (const FSoftObjectPath& StartObject : Start)
+				for (const TSoftObjectPtr<>& StartObject : Start)
 				{
-					const TArray<FSoftObjectPath>* Children = Hierarchy.Find(StartObject);
+					const TArray<FSoftObjectPath>* Children = Hierarchy.Find(StartObject.GetUniqueID());
 					if (!Children)
 					{
 						continue;
@@ -86,7 +86,7 @@ namespace UE::ConcertSharedSlate
 					const FSoftObjectPath ExpectedChildPath(*ExpectedChildPathString);
 					if (Children->Contains(ExpectedChildPath))
 					{
-						RelatedChildren.Add(ExpectedChildPath);
+						RelatedChildren.Emplace(ExpectedChildPath);
 					}
 				}
 
@@ -111,7 +111,7 @@ namespace UE::ConcertSharedSlate
 		];
 	}
 
-	void SMultiObjectAssignment::RefreshData(const TArray<FSoftObjectPath>& Objects, const IReplicationStreamModel& Model)
+	void SMultiObjectAssignment::RefreshData(const TArray<TSoftObjectPtr<>>& Objects, const IReplicationStreamModel& Model)
 	{
 		TArray<FPropertyAssignmentEntry> Entries;
 		auto[RootEntry, bRootSharesClass] = BuildAssignmentEntry(Objects, Model);
@@ -126,7 +126,7 @@ namespace UE::ConcertSharedSlate
 		{
 			using namespace Private;
 			const TMap<FSoftObjectPath, TArray<FSoftObjectPath>> HierarchyInfo = BuildHierarchy(Objects, *ObjectHierarchy, bShouldShowSubobjects);
-			ForEachCategory(Objects, HierarchyInfo, [this, &Model, &Entries](const TArray<FSoftObjectPath>& ContextObjects)
+			ForEachCategory(Objects, HierarchyInfo, [this, &Model, &Entries](const TArray<TSoftObjectPtr<>>& ContextObjects)
 			{
 				auto[Entry, bShareClass] = BuildAssignmentEntry(ContextObjects, Model);
 				if (bShareClass)
@@ -156,14 +156,11 @@ namespace UE::ConcertSharedSlate
 		}
 	}
 
-	SMultiObjectAssignment::FBuildAssignmentEntryResult SMultiObjectAssignment::BuildAssignmentEntry(const TArray<FSoftObjectPath>& Objects, const IReplicationStreamModel& Model)
+	SMultiObjectAssignment::FBuildAssignmentEntryResult SMultiObjectAssignment::BuildAssignmentEntry(const TArray<TSoftObjectPtr<>>& Objects, const IReplicationStreamModel& Model)
 	{
 		FBuildAssignmentEntryResult Result;
-		
-		TArray<TSoftObjectPtr<>> ContextObjects;
-		Algo::Transform(Objects, ContextObjects, [](const FSoftObjectPath& Path){ return TSoftObjectPtr{ Path }; });
 
-		Result.Entry = { .ContextObjects = MoveTemp(ContextObjects) };
+		Result.Entry = { .ContextObjects = Objects };
 		TSet<FConcertPropertyChain>& Properties = Result.Entry.PropertiesToDisplay;
 		FSoftClassPath& ClassPath = Result.Entry.Class;
 		
diff --git a/Engine/Plugins/Developer/Concert/ConcertUI/ConcertSharedSlate/Source/ConcertSharedSlate/Private/Replication/Editor/View/Property/SMultiObjectAssignment.h b/Engine/Plugins/Developer/Concert/ConcertUI/ConcertSharedSlate/Source/ConcertSharedSlate/Private/Replication/Editor/View/Property/SMultiObjectAssignment.h
--- a/Engine/Plugins/Developer/Concert/ConcertUI/ConcertSharedSlate/Source/ConcertSharedSlate/Private/Replication/Editor/View/Property/SMultiObjectAssignment.h
+++ b/Engine/Plugins/Developer/Concert/ConcertUI/ConcertSharedSlate/Source/ConcertSharedSlate/Private/Replication/Editor/View/Property/SMultiObjectAssignment.h
@@ -9,6 +9,9 @@
 #include "Widgets/DeclarativeSyntaxSupport.h"
 #include "Widgets/SCompoundWidget.h"
 
+template<class T>
+struct TSoftObjectPtr;
+
 namespace UE::ConcertSharedSlate
 {
 	class IObjectHierarchyModel;
@@ -33,7 +36,7 @@ namespace UE::ConcertSharedSlate
 		void Construct(const FArguments& InArgs, TSharedRef<IPropertyTreeView> InTreeView);
 
 		//~ Begin IPropertyAssignmentView Interface
-		virtual void RefreshData(const TArray<FSoftObjectPath>& Objects, const IReplicationStreamModel& Model) override;
+		virtual void RefreshData(const TArray<TSoftObjectPtr<>>& Objects, const IReplicationStreamModel& Model) override;
 		virtual void RequestRefilter() const override { return TreeView->RequestRefilter(); }
 		virtual void RequestResortForColumn(const FName& ColumnId) override { return TreeView->RequestRefilter(); }
 		virtual TSharedRef<SWidget> GetWidget() override { return SharedThis(this); }
@@ -46,7 +49,7 @@ namespace UE::ConcertSharedSlate
 		/** The tree view that is being wrapped. */
 		TSharedPtr<IPropertyTreeView> TreeView;
 		/** Used to determine whether to rebuild the entire property data. */
-		TArray<FSoftObjectPath> PreviousSelectedObjects;
+		TArray<TSoftObjectPtr<>> PreviousSelectedObjects;
 
 		/** Used to get subobjects of selected objects. */
 		TSharedPtr<IObjectHierarchyModel> ObjectHierarchy;
@@ -62,6 +65,6 @@ namespace UE::ConcertSharedSlate
 			bool bHaveSharedClass = false;
 		};
 		/** Builds a property section grouped by Objects (usually has 1 object - contains similar objects, like StaticMeshComponent0, for multi-edit purposes). */
-		FBuildAssignmentEntryResult BuildAssignmentEntry(const TArray<FSoftObjectPath>& Objects, const IReplicationStreamModel& Model);
+		FBuildAssignmentEntryResult BuildAssignmentEntry(const TArray<TSoftObjectPtr<>>& Objects, const IReplicationStreamModel& Model);
 	};
 }
diff --git a/Engine/Plugins/Developer/Concert/ConcertUI/ConcertSharedSlate/Source/ConcertSharedSlate/Private/Replication/Editor/View/Property/SPerObjectPropertyAssignment.cpp b/Engine/Plugins/Developer/Concert/ConcertUI/ConcertSharedSlate/Source/ConcertSharedSlate/Private/Replication/Editor/View/Property/SPerObjectPropertyAssignment.cpp
--- a/Engine/Plugins/Developer/Concert/ConcertUI/ConcertSharedSlate/Source/ConcertSharedSlate/Private/Replication/Editor/View/Property/SPerObjectPropertyAssignment.cpp
+++ b/Engine/Plugins/Developer/Concert/ConcertUI/ConcertSharedSlate/Source/ConcertSharedSlate/Private/Replication/Editor/View/Property/SPerObjectPropertyAssignment.cpp
@@ -19,12 +19,9 @@ namespace UE::ConcertSharedSlate
 		];
 	}
 
-	void SPerObjectPropertyAssignment::RefreshData(const TArray<FSoftObjectPath>& Objects, const IReplicationStreamModel& Model)
+	void SPerObjectPropertyAssignment::RefreshData(const TArray<TSoftObjectPtr<>>& Objects, const IReplicationStreamModel& Model)
 	{
-		TArray<TSoftObjectPtr<>> ContextObjects;
-		Algo::Transform(Objects, ContextObjects, [](const FSoftObjectPath& Path){ return TSoftObjectPtr{ Path }; });
-		
-		FPropertyAssignmentEntry AssignmentEntry { .ContextObjects = MoveTemp(ContextObjects) };
+		FPropertyAssignmentEntry AssignmentEntry { .ContextObjects = Objects };
 		TSet<FConcertPropertyChain>& Properties = AssignmentEntry.PropertiesToDisplay;
 		FSoftClassPath& ClassPath = AssignmentEntry.Class;
 		
diff --git a/Engine/Plugins/Developer/Concert/ConcertUI/ConcertSharedSlate/Source/ConcertSharedSlate/Private/Replication/Editor/View/Property/SPerObjectPropertyAssignment.h b/Engine/Plugins/Developer/Concert/ConcertUI/ConcertSharedSlate/Source/ConcertSharedSlate/Private/Replication/Editor/View/Property/SPerObjectPropertyAssignment.h
--- a/Engine/Plugins/Developer/Concert/ConcertUI/ConcertSharedSlate/Source/ConcertSharedSlate/Private/Replication/Editor/View/Property/SPerObjectPropertyAssignment.h
+++ b/Engine/Plugins/Developer/Concert/ConcertUI/ConcertSharedSlate/Source/ConcertSharedSlate/Private/Replication/Editor/View/Property/SPerObjectPropertyAssignment.h
@@ -28,7 +28,7 @@ namespace UE::ConcertSharedSlate
 		void Construct(const FArguments& InArgs, TSharedRef<IPropertyTreeView> InTreeView);
 
 		//~ Begin IPropertyAssignmentView Interface
-		virtual void RefreshData(const TArray<FSoftObjectPath>& Objects, const IReplicationStreamModel& Model) override;
+		virtual void RefreshData(const TArray<TSoftObjectPtr<>>& Objects, const IReplicationStreamModel& Model) override;
 		virtual void RequestRefilter() const override { return TreeView->RequestRefilter(); }
 		virtual void RequestResortForColumn(const FName& ColumnId) override { return TreeView->RequestRefilter(); }
 		virtual TSharedRef<SWidget> GetWidget() override { return SharedThis(this); }
@@ -39,7 +39,7 @@ namespace UE::ConcertSharedSlate
 		/** The tree view that is being wrapped. */
 		TSharedPtr<IPropertyTreeView> TreeView;
 		/** Used to determine whether to rebuild the entire property data. */
-		TArray<FSoftObjectPath> PreviousSelectedObjects;
+		TArray<TSoftObjectPtr<>> PreviousSelectedObjects;
 		
 		/** Optional. If specified, displays the properties of this model instead of those assigned in the stream. */
 		TSharedPtr<IPropertySelectionSourceModel> OptionalPropertySource;
diff --git a/Engine/Plugins/Developer/Concert/ConcertUI/ConcertSharedSlate/Source/ConcertSharedSlate/Public/Replication/Editor/Model/Data/ReplicatedObjectData.h b/Engine/Plugins/Developer/Concert/ConcertUI/ConcertSharedSlate/Source/ConcertSharedSlate/Public/Replication/Editor/Model/Data/ReplicatedObjectData.h
--- a/Engine/Plugins/Developer/Concert/ConcertUI/ConcertSharedSlate/Source/ConcertSharedSlate/Public/Replication/Editor/Model/Data/ReplicatedObjectData.h
+++ b/Engine/Plugins/Developer/Concert/ConcertUI/ConcertSharedSlate/Source/ConcertSharedSlate/Public/Replication/Editor/Model/Data/ReplicatedObjectData.h
@@ -15,6 +15,9 @@ namespace UE::ConcertSharedSlate
 		FReplicatedObjectData(FSoftObjectPath ObjectPath)
 			: ObjectPtr(MoveTemp(ObjectPath))
 		{}
+		FReplicatedObjectData(TSoftObjectPtr<> Object)
+			: ObjectPtr(MoveTemp(Object))
+		{}
 		
 		const FSoftObjectPath& GetObjectPath() const { return ObjectPtr.GetUniqueID(); }
 		const TSoftObjectPtr<>& GetObjectPtr() const { return ObjectPtr; }
diff --git a/Engine/Plugins/Developer/Concert/ConcertUI/ConcertSharedSlate/Source/ConcertSharedSlate/Public/Replication/Editor/Model/Object/IObjectHierarchyModel.h b/Engine/Plugins/Developer/Concert/ConcertUI/ConcertSharedSlate/Source/ConcertSharedSlate/Public/Replication/Editor/Model/Object/IObjectHierarchyModel.h
--- a/Engine/Plugins/Developer/Concert/ConcertUI/ConcertSharedSlate/Source/ConcertSharedSlate/Public/Replication/Editor/Model/Object/IObjectHierarchyModel.h
+++ b/Engine/Plugins/Developer/Concert/ConcertUI/ConcertSharedSlate/Source/ConcertSharedSlate/Public/Replication/Editor/Model/Object/IObjectHierarchyModel.h
@@ -5,8 +5,9 @@
 #include "Containers/Array.h"
 #include "Delegates/Delegate.h"
 #include "Misc/EBreakBehavior.h"
-#include "UObject/SoftObjectPath.h"
 #include "Templates/Function.h"
+#include "UObject/SoftObjectPath.h"
+#include "UObject/SoftObjectPtr.h"
 
 enum class EBreakBehavior : uint8;
 
@@ -68,35 +69,35 @@ namespace UE::ConcertSharedSlate
 		 * @param InclusionFlags Specifies the relationships that should be included
 		 */
 		virtual void ForEachDirectChild(
-			const FSoftObjectPath& Root,
-			TFunctionRef<EBreakBehavior(const FSoftObjectPath& Object, EChildRelationship Relationship)> Callback,
+			const TSoftObjectPtr<>& Root,
+			TFunctionRef<EBreakBehavior(const TSoftObjectPtr<>& Object, EChildRelationship Relationship)> Callback,
 			EChildRelationshipFlags InclusionFlags = EChildRelationshipFlags::All
 			) const = 0;
 
 		struct FParentInfo
 		{
 			/** The parent of the child object in the hierarchy. */
-			FSoftObjectPath Parent;
+			TSoftObjectPtr<> Parent;
 			/** Relationship that the child has to Parent, e.g. the child could be a component of Parent. */
 			EChildRelationship Relationship;
 		};
 		/** Gets parent info for ChildObject, if it has a child. */
-		virtual TOptional<FParentInfo> GetParentInfo(const FSoftObjectPath& ChildObject) const = 0;
+		virtual TOptional<FParentInfo> GetParentInfo(const TSoftObjectPtr<>& ChildObject) const = 0;
 
 		/** Util for iterating all subobjects. */
 		void ForEachChildRecursive(
-			const FSoftObjectPath& Root,
-			TFunctionRef<EBreakBehavior(const FSoftObjectPath& Parent, const FSoftObjectPath& ChildObject, EChildRelationship Relationship)> Callback,
+			const TSoftObjectPtr<>& Root,
+			TFunctionRef<EBreakBehavior(const TSoftObjectPtr<>& Parent, const TSoftObjectPtr<>& ChildObject, EChildRelationship Relationship)> Callback,
 			EChildRelationshipFlags InclusionFlags = EChildRelationshipFlags::All
 			) const;
 
 		/** Builds a TArray containing all children, recursively. */
 		template<typename TAllocator = FDefaultAllocator>
-		TArray<FSoftObjectPath, TAllocator> GetChildrenRecursive(const FSoftObjectPath& Root, EChildRelationshipFlags InclusionFlags = EChildRelationshipFlags::All) const
+		TArray<TSoftObjectPtr<>, TAllocator> GetChildrenRecursive(const TSoftObjectPtr<>& Root, EChildRelationshipFlags InclusionFlags = EChildRelationshipFlags::All) const
 		{
-			TArray<FSoftObjectPath, TAllocator> AllObjects;
+			TArray<TSoftObjectPtr<>, TAllocator> AllObjects;
 			ForEachChildRecursive(Root,
-				[&AllObjects](const FSoftObjectPath&, const FSoftObjectPath& ChildObject, EChildRelationship)
+				[&AllObjects](const TSoftObjectPtr<>&, const TSoftObjectPtr<>& ChildObject, EChildRelationship)
 				{
 					AllObjects.Add(ChildObject);
 					return EBreakBehavior::Continue;
diff --git a/Engine/Plugins/Developer/Concert/ConcertUI/ConcertSharedSlate/Source/ConcertSharedSlate/Public/Replication/Editor/Model/Property/IPropertySelectionSourceModel.h b/Engine/Plugins/Developer/Concert/ConcertUI/ConcertSharedSlate/Source/ConcertSharedSlate/Public/Replication/Editor/Model/Property/IPropertySelectionSourceModel.h
--- a/Engine/Plugins/Developer/Concert/ConcertUI/ConcertSharedSlate/Source/ConcertSharedSlate/Public/Replication/Editor/Model/Property/IPropertySelectionSourceModel.h
+++ b/Engine/Plugins/Developer/Concert/ConcertUI/ConcertSharedSlate/Source/ConcertSharedSlate/Public/Replication/Editor/Model/Property/IPropertySelectionSourceModel.h
@@ -4,11 +4,28 @@
 
 #include "IPropertySourceModel.h"
 #include "Model/Item/SourceSelectionCategory.h"
+
 #include "Templates/SharedPointer.h"
+#include "UObject/SoftObjectPath.h"
+#include "UObject/SoftObjectPtr.h"
 
 namespace UE::ConcertSharedSlate
 {
 	using FPropertySourceCategory = ConcertSharedSlate::TSourceSelectionCategory<FSelectablePropertyInfo>;
+
+	/** Arguments for getting properties associated with an object / class. */
+	struct FPropertySourceContext
+	{
+		/** The object for which the properties are supposed to be displayed */
+		TSoftObjectPtr<> Object;
+		/** The class of Object */
+		FSoftClassPath Class;
+
+		FPropertySourceContext(const TSoftObjectPtr<>& Object, const FSoftClassPath& Class)
+			: Object(Object)
+			, Class(Class)
+		{}
+	};
 	
 	/** Decides which properties can be added to a IReplicationStreamModel. */
 	class IPropertySelectionSourceModel : public TSharedFromThis<IPropertySelectionSourceModel>
@@ -16,7 +33,11 @@ namespace UE::ConcertSharedSlate
 	public:
 
 		/** Gets the single source determining which properties can be selected. */
-		virtual TSharedRef<IPropertySourceModel> GetPropertySource(const FSoftClassPath& Class) const = 0;
+		UE_DEPRECATED(5.5, "Use the version that takes FPropertySourceContext instead.")
+		virtual TSharedRef<IPropertySourceModel> GetPropertySource(const FSoftClassPath& Class) const { return GetPropertySource(FPropertySourceContext({}, Class)); }
+
+		/** @return Gets an object that will tell you the properties that are associated with the passed in object / class context. */
+		virtual TSharedRef<IPropertySourceModel> GetPropertySource(const FPropertySourceContext& Context) const = 0;
 		
 		virtual ~IPropertySelectionSourceModel() = default;
 	};
diff --git a/Engine/Plugins/Developer/Concert/ConcertUI/ConcertSharedSlate/Source/ConcertSharedSlate/Public/Replication/Editor/View/IPropertyAssignmentView.h b/Engine/Plugins/Developer/Concert/ConcertUI/ConcertSharedSlate/Source/ConcertSharedSlate/Public/Replication/Editor/View/IPropertyAssignmentView.h
--- a/Engine/Plugins/Developer/Concert/ConcertUI/ConcertSharedSlate/Source/ConcertSharedSlate/Public/Replication/Editor/View/IPropertyAssignmentView.h
+++ b/Engine/Plugins/Developer/Concert/ConcertUI/ConcertSharedSlate/Source/ConcertSharedSlate/Public/Replication/Editor/View/IPropertyAssignmentView.h
@@ -5,9 +5,10 @@
 #include "Containers/ContainersFwd.h"
 #include "Containers/Set.h"
 #include "Templates/SharedPointer.h"
+#include "UObject/SoftObjectPtr.h"
 
 class SWidget;
-struct FSoftObjectPath;
+class UObject;
 
 namespace UE::ConcertSharedSlate
 {
@@ -35,7 +36,7 @@ namespace UE::ConcertSharedSlate
 		 * @param Objects The objects that are supposed to be displayed.
 		 * @param Model The model that can be queried for object info.
 		 */
-		virtual void RefreshData(const TArray<FSoftObjectPath>& Objects, const IReplicationStreamModel& Model) = 0;
+		virtual void RefreshData(const TArray<TSoftObjectPtr<>>& Objects, const IReplicationStreamModel& Model) = 0;
 		
 		/** Reapply the filter function to all items at the end of the frame. Call e.g. when the filters have changed. */
 		virtual void RequestRefilter() const = 0;
diff --git a/Engine/Plugins/Developer/Concert/ConcertUI/ConcertSharedSlate/Source/ConcertSharedSlate/Public/Replication/Editor/View/IReplicationStreamViewer.h b/Engine/Plugins/Developer/Concert/ConcertUI/ConcertSharedSlate/Source/ConcertSharedSlate/Public/Replication/Editor/View/IReplicationStreamViewer.h
--- a/Engine/Plugins/Developer/Concert/ConcertUI/ConcertSharedSlate/Source/ConcertSharedSlate/Public/Replication/Editor/View/IReplicationStreamViewer.h
+++ b/Engine/Plugins/Developer/Concert/ConcertUI/ConcertSharedSlate/Source/ConcertSharedSlate/Public/Replication/Editor/View/IReplicationStreamViewer.h
@@ -29,13 +29,13 @@ namespace UE::ConcertSharedSlate
 		 */
 		virtual void RequestPropertyColumnResort(const FName& ColumnId) = 0;
 
-		/**
-		 * @return The objects for which the properties are being edited / displayed.
-		 * If there is an IReplicationSubobjectView, this is IReplicationSubobjectView::GetSelectedObjects.
-		 * Otherwise it is IReplicationStreamViewer::GetSelectedTopLevelObjects.
-		 */
-		virtual TArray<FSoftObjectPath> GetObjectsBeingPropertyEdited() const = 0;
+		/** @return The objects for which the properties are being edited / displayed. */
+		UE_DEPRECATED(5.5, "Use GetSelectedObjects instead.")
+		virtual TArray<FSoftObjectPath> GetObjectsBeingPropertyEdited() const;
+
+		/** @return The objects for which the properties are being edited / displayed. */
+		virtual TArray<TSoftObjectPtr<>> GetSelectedObjects() const = 0;
 
 		virtual ~IReplicationStreamViewer() = default;
 	};
-}
\ No newline at end of file
+}
diff --git a/Engine/Plugins/Developer/Concert/ConcertUI/ConcertSharedSlate/Source/ConcertSharedSlate/Public/Replication/Utils/ReplicationWidgetDelegates.h b/Engine/Plugins/Developer/Concert/ConcertUI/ConcertSharedSlate/Source/ConcertSharedSlate/Public/Replication/Utils/ReplicationWidgetDelegates.h
--- a/Engine/Plugins/Developer/Concert/ConcertUI/ConcertSharedSlate/Source/ConcertSharedSlate/Public/Replication/Utils/ReplicationWidgetDelegates.h
+++ b/Engine/Plugins/Developer/Concert/ConcertUI/ConcertSharedSlate/Source/ConcertSharedSlate/Public/Replication/Utils/ReplicationWidgetDelegates.h
@@ -24,7 +24,7 @@ namespace UE::ConcertSharedSlate
 	DECLARE_DELEGATE_RetVal_TwoParams(bool, FSortPropertyPredicate, const FPropertyData& Left, const FPropertyData& Right);
 
 	/** Extends a context menu that is being built for a selection of objects. */
-	DECLARE_DELEGATE_TwoParams(FExtendObjectMenu, FMenuBuilder&, TConstArrayView<FSoftObjectPath> ContextObjects);
+	DECLARE_DELEGATE_TwoParams(FExtendObjectMenu, FMenuBuilder&, TConstArrayView<TSoftObjectPtr<>> ContextObjects);
 
 	/** Delegate for getting an object's class. */
 	DECLARE_DELEGATE_RetVal_OneParam(FSoftClassPath, FGetObjectClass, const TSoftObjectPtr<>&);
diff --git a/Engine/Plugins/Editor/ContentBrowser/ContentBrowserAssetDataSource/Source/ContentBrowserAssetDataSource/Private/ContentBrowserAssetDataCore.cpp b/Engine/Plugins/Editor/ContentBrowser/ContentBrowserAssetDataSource/Source/ContentBrowserAssetDataSource/Private/ContentBrowserAssetDataCore.cpp
--- a/Engine/Plugins/Editor/ContentBrowser/ContentBrowserAssetDataSource/Source/ContentBrowserAssetDataSource/Private/ContentBrowserAssetDataCore.cpp
+++ b/Engine/Plugins/Editor/ContentBrowser/ContentBrowserAssetDataSource/Source/ContentBrowserAssetDataSource/Private/ContentBrowserAssetDataCore.cpp
@@ -1899,7 +1899,7 @@ bool GetAssetDataAttribute(const FAssetData& InAssetData, const bool InIncludeMe
 		FAssetDataTagMapSharedView::FFindTagResult FoundValue = InAssetData.TagsAndValues.FindTag(FoundAttributeKey);
 		if (!FoundValue.IsSet())
 		{
-			ensureMsgf(ClassPropertyTagCache || (FindObject<UClass>(InAssetData.AssetClassPath) == nullptr), TEXT("FAssetPropertyTagCache not populated for type %s when looking for attribute %s"), *WriteToString<256>(InAssetData.AssetClassPath), *WriteToString<256>(InAttributeKey));
+			//ensureMsgf(ClassPropertyTagCache || (FindObject<UClass>(InAssetData.AssetClassPath) == nullptr), TEXT("FAssetPropertyTagCache not populated for type %s when looking for attribute %s"), *WriteToString<256>(InAssetData.AssetClassPath), *WriteToString<256>(InAttributeKey));
 			if (ClassPropertyTagCache)
 			{
 				// Check to see if the key we were given resolves as an alias
@@ -1998,8 +1998,8 @@ bool GetAssetDataAttributes(const FAssetData& InAssetData, const bool InIncludeM
 			}
 		}
 
-		ensureMsgf(!InIncludeMetaData || ClassPropertyTagCache || (FindObject<UClass>(InAssetData.AssetClassPath) == nullptr), TEXT("FAssetPropertyTagCache not populated for type %s when fetching all attributes"), 
-			*WriteToString<256>(InAssetData.AssetClassPath));
+		/*ensureMsgf(!InIncludeMetaData || ClassPropertyTagCache || (FindObject<UClass>(InAssetData.AssetClassPath) == nullptr), TEXT("FAssetPropertyTagCache not populated for type %s when fetching all attributes"), 
+			*WriteToString<256>(InAssetData.AssetClassPath));*/
 
 		OutAttributeValues.Reserve(OutAttributeValues.Num() + InAssetData.TagsAndValues.Num());
 		for (const auto& TagAndValue : InAssetData.TagsAndValues)
