Deprecate IPropertySelectionSourceModel and replace with new IPropertySourceProcessor, which has a changed, "safer" signature.

diff --git a/Engine/Plugins/Developer/Concert/ConcertUI/ConcertClientSharedSlate/Source/ConcertClientSharedSlate/Private/Replication/Editor/Model/PropertySource/SelectPropertyFromUClassModel.cpp b/Engine/Plugins/Developer/Concert/ConcertUI/ConcertClientSharedSlate/Source/ConcertClientSharedSlate/Private/Replication/Editor/Model/PropertySource/SelectPropertyFromUClassModel.cpp
--- a/Engine/Plugins/Developer/Concert/ConcertUI/ConcertClientSharedSlate/Source/ConcertClientSharedSlate/Private/Replication/Editor/Model/PropertySource/SelectPropertyFromUClassModel.cpp
+++ b/Engine/Plugins/Developer/Concert/ConcertUI/ConcertClientSharedSlate/Source/ConcertClientSharedSlate/Private/Replication/Editor/Model/PropertySource/SelectPropertyFromUClassModel.cpp
@@ -2,30 +2,26 @@
 
 #include "Replication/Editor/Model/PropertySource/SelectPropertyFromUClassModel.h"
 
-#include "ConcertLogGlobal.h"
+#include "Replication/Editor/Model/Property/IPropertySourceModel.h"
 #include "Replication/Editor/Model/PropertySource/ConcertSyncCoreReplicatedPropertySource.h"
 
+#include "ConcertLogGlobal.h"
+
 #define LOCTEXT_NAMESPACE "FSelectPropertyFromUClassModel"
 
 namespace UE::ConcertClientSharedSlate
 {
-	FSelectPropertyFromUClassModel::FSelectPropertyFromUClassModel()
-		: UClassIteratorSource(MakeShared<FConcertSyncCoreReplicatedPropertySource>())
-	{}
-
-	TSharedRef<ConcertSharedSlate::IPropertySourceModel> FSelectPropertyFromUClassModel::GetPropertySource(const ConcertSharedSlate::FPropertySourceContext& Context) const
+	void FSelectPropertyFromUClassModel::ProcessPropertySource(
+		const ConcertSharedSlate::FPropertySourceContext& Context,
+		TFunctionRef<void(const ConcertSharedSlate::IPropertySourceModel& Model)> Processor
+		) const
 	{
 		UClass* LoadedClass = Context.Class.TryLoadClass<UObject>();
-		if (LoadedClass)
-		{
-			UClassIteratorSource->SetClass(LoadedClass);
-		}
-		else
-		{
-			UClassIteratorSource->SetClass(nullptr);
-			UE_LOG(LogConcert, Warning, TEXT("Could not resolve class %s. Properties will not be available."), *Context.Class.ToString());
-		}
-		return UClassIteratorSource;
+		
+		FConcertSyncCoreReplicatedPropertySource UClassIteratorSource;
+		UClassIteratorSource.SetClass(LoadedClass ? LoadedClass : nullptr);
+		UE_CLOG(!LoadedClass, LogConcert, Warning, TEXT("Could not resolve class %s. Properties will not be available."), *Context.Class.ToString());
+		Processor(UClassIteratorSource);
 	}
 }
 
diff --git a/Engine/Plugins/Developer/Concert/ConcertUI/ConcertClientSharedSlate/Source/ConcertClientSharedSlate/Public/Replication/ClientReplicationWidgetFactories.h b/Engine/Plugins/Developer/Concert/ConcertUI/ConcertClientSharedSlate/Source/ConcertClientSharedSlate/Public/Replication/ClientReplicationWidgetFactories.h
--- a/Engine/Plugins/Developer/Concert/ConcertUI/ConcertClientSharedSlate/Source/ConcertClientSharedSlate/Public/Replication/ClientReplicationWidgetFactories.h
+++ b/Engine/Plugins/Developer/Concert/ConcertUI/ConcertClientSharedSlate/Source/ConcertClientSharedSlate/Public/Replication/ClientReplicationWidgetFactories.h
@@ -11,7 +11,7 @@ struct FConcertObjectReplicationMap;
 
 namespace UE::ConcertSharedSlate
 {
-	class IPropertySelectionSourceModel;
+	class IPropertySourceProcessor;
 	class IPropertyTreeView;
 	class IEditableReplicationStreamModel;
 	class IObjectNameModel;
diff --git a/Engine/Plugins/Developer/Concert/ConcertUI/ConcertClientSharedSlate/Source/ConcertClientSharedSlate/Public/Replication/Editor/Model/PropertySource/SelectPropertyFromUClassModel.h b/Engine/Plugins/Developer/Concert/ConcertUI/ConcertClientSharedSlate/Source/ConcertClientSharedSlate/Public/Replication/Editor/Model/PropertySource/SelectPropertyFromUClassModel.h
--- a/Engine/Plugins/Developer/Concert/ConcertUI/ConcertClientSharedSlate/Source/ConcertClientSharedSlate/Public/Replication/Editor/Model/PropertySource/SelectPropertyFromUClassModel.h
+++ b/Engine/Plugins/Developer/Concert/ConcertUI/ConcertClientSharedSlate/Source/ConcertClientSharedSlate/Public/Replication/Editor/Model/PropertySource/SelectPropertyFromUClassModel.h
@@ -2,33 +2,30 @@
 
 #pragma once
 
-#include "Replication/Editor/Model/Property/IPropertySelectionSourceModel.h"
-
+#include "ConcertSyncCoreReplicatedPropertySource.h"
+#include "Replication/Editor/Model/Property/IPropertySourceProcessor.h"
 #include "Templates/SharedPointer.h"
 
 struct FSoftClassPath;
 
 namespace UE::ConcertClientSharedSlate
 {
-	class FConcertSyncCoreReplicatedPropertySource;
 	
 	/**
 	 * Decides which properties can be added to IEditableReplicationStreamModel.
 	 * The allowed properties are those returned by UE::ConcertSyncCore::ForEachReplicatableProperty.
 	 */
-	class CONCERTCLIENTSHAREDSLATE_API FSelectPropertyFromUClassModel : public ConcertSharedSlate::IPropertySelectionSourceModel
+	class CONCERTCLIENTSHAREDSLATE_API FSelectPropertyFromUClassModel : public ConcertSharedSlate::IPropertySourceProcessor
 	{
 	public:
 
-		FSelectPropertyFromUClassModel();
-
-		//~ Begin IPropertySelectionSourceModel Interface
-		virtual TSharedRef<ConcertSharedSlate::IPropertySourceModel> GetPropertySource(const ConcertSharedSlate::FPropertySourceContext& Context) const override;
-		//~ End IPropertySelectionSourceModel Interface
-
-	private:
+		//~ Begin IPropertySelectionProcessor Interface
+		virtual void ProcessPropertySource(
+			const ConcertSharedSlate::FPropertySourceContext& Context,
+			TFunctionRef<void(const ConcertSharedSlate::IPropertySourceModel& Model)> Processor
+			) const override;
+		//~ End IPropertySelectionProcessor Interface
 
-		TSharedRef<FConcertSyncCoreReplicatedPropertySource> UClassIteratorSource;
 	};
 }
 
diff --git a/Engine/Plugins/Developer/Concert/ConcertUI/ConcertSharedSlate/Source/ConcertSharedSlate/Private/Replication/Editor/Utils/PropertyEnumerationUtils.cpp b/Engine/Plugins/Developer/Concert/ConcertUI/ConcertSharedSlate/Source/ConcertSharedSlate/Private/Replication/Editor/Utils/PropertyEnumerationUtils.cpp
--- a/Engine/Plugins/Developer/Concert/ConcertUI/ConcertSharedSlate/Source/ConcertSharedSlate/Private/Replication/Editor/Utils/PropertyEnumerationUtils.cpp
+++ b/Engine/Plugins/Developer/Concert/ConcertUI/ConcertSharedSlate/Source/ConcertSharedSlate/Private/Replication/Editor/Utils/PropertyEnumerationUtils.cpp
@@ -4,7 +4,7 @@
 
 #include "Misc/EBreakBehavior.h"
 #include "Replication/Editor/Model/IReplicationStreamModel.h"
-#include "Replication/Editor/Model/Property/IPropertySelectionSourceModel.h"
+#include "Replication/Editor/Model/Property/IPropertySourceProcessor.h"
 
 #include "Containers/Set.h"
 #include "UObject/SoftObjectPath.h"
@@ -29,7 +29,7 @@ namespace UE::ConcertSharedSlate
 #endif
 	}
 
-	void EnumerateProperties(TConstArrayView<TSoftObjectPtr<>> Objects, const IReplicationStreamModel& Model, const IPropertySelectionSourceModel* OptionalSource, FEnumerateProperties Callback)
+	void EnumerateProperties(TConstArrayView<TSoftObjectPtr<>> Objects, const IReplicationStreamModel& Model, const IPropertySourceProcessor* OptionalSource, FEnumerateProperties Callback)
 	{
 		if (OptionalSource)
 		{
@@ -63,7 +63,7 @@ namespace UE::ConcertSharedSlate
 	}
 
 	/** Enumerate the properties that are selectable in Source (e.g. all properties in that class, @see FSelectPropertyFromUClassModel). */
-	void EnumerateAllProperties(TConstArrayView<TSoftObjectPtr<>> Objects, const IPropertySelectionSourceModel& Source, const IReplicationStreamModel& Model, FEnumerateProperties Callback)
+	void EnumerateAllProperties(TConstArrayView<TSoftObjectPtr<>> Objects, const IPropertySourceProcessor& Source, const IReplicationStreamModel& Model, FEnumerateProperties Callback)
 	{
 		TSet<FSoftClassPath> VisitedClasses; 
 			
@@ -73,10 +73,13 @@ namespace UE::ConcertSharedSlate
 			const FPropertySourceContext ObjectQueryContext(Object, ObjectClass);
 					
 			EBreakBehavior BreakBehavior = EBreakBehavior::Continue;
-			Source.GetPropertySource(ObjectQueryContext)->EnumerateSelectableItems([&Callback, &ObjectClass, &BreakBehavior](const FSelectablePropertyInfo& PropertyInfo)
+			Source.ProcessPropertySource(ObjectQueryContext, [&Callback, &ObjectClass, &BreakBehavior](const IPropertySourceModel& PropertySource)
 			{
-				BreakBehavior = Callback(ObjectClass, PropertyInfo.Property);
-				return BreakBehavior;
+				PropertySource.EnumerateSelectableItems([&Callback, &ObjectClass, &BreakBehavior](const FSelectablePropertyInfo& PropertyInfo)
+				{
+					BreakBehavior = Callback(ObjectClass, PropertyInfo.Property);
+					return BreakBehavior;
+				});
 			});
 					
 			if (BreakBehavior == EBreakBehavior::Break)
diff --git a/Engine/Plugins/Developer/Concert/ConcertUI/ConcertSharedSlate/Source/ConcertSharedSlate/Private/Replication/Editor/Utils/PropertyEnumerationUtils.h b/Engine/Plugins/Developer/Concert/ConcertUI/ConcertSharedSlate/Source/ConcertSharedSlate/Private/Replication/Editor/Utils/PropertyEnumerationUtils.h
--- a/Engine/Plugins/Developer/Concert/ConcertUI/ConcertSharedSlate/Source/ConcertSharedSlate/Private/Replication/Editor/Utils/PropertyEnumerationUtils.h
+++ b/Engine/Plugins/Developer/Concert/ConcertUI/ConcertSharedSlate/Source/ConcertSharedSlate/Private/Replication/Editor/Utils/PropertyEnumerationUtils.h
@@ -14,7 +14,7 @@ struct FSoftObjectPath;
 
 namespace UE::ConcertSharedSlate
 {
-	class IPropertySelectionSourceModel;
+	class IPropertySourceProcessor;
 	class IReplicationStreamModel;
 	using FEnumerateProperties = TFunctionRef<EBreakBehavior(const FSoftClassPath&, const FConcertPropertyChain&)>;
 	
@@ -22,9 +22,9 @@ namespace UE::ConcertSharedSlate
 	FSoftClassPath GetObjectClassFromModelOrLoad(const TSoftObjectPtr<>& Object, const IReplicationStreamModel& Model);
 
 	/** Calls EnumerateRegisteredPropertiesOnly or EnumerateAllProperties depending on whether OptionalSource is nullptr. */
-	void EnumerateProperties(TConstArrayView<TSoftObjectPtr<>> Objects, const IReplicationStreamModel& Model, const IPropertySelectionSourceModel* OptionalSource, FEnumerateProperties Callback);
+	void EnumerateProperties(TConstArrayView<TSoftObjectPtr<>> Objects, const IReplicationStreamModel& Model, const IPropertySourceProcessor* OptionalSource, FEnumerateProperties Callback);
 	/** Enumerates the properties that are assigned to the object in Model */
 	void EnumerateRegisteredPropertiesOnly(TConstArrayView<TSoftObjectPtr<>> Objects, const IReplicationStreamModel& Model, FEnumerateProperties Callback);
 	/** Enumerate the properties that are selectable in Source (e.g. all properties in that class, @see FSelectPropertyFromUClassModel). */
-	void EnumerateAllProperties(TConstArrayView<TSoftObjectPtr<>> Objects, const IPropertySelectionSourceModel& Source, const IReplicationStreamModel& Model, FEnumerateProperties Callback);
+	void EnumerateAllProperties(TConstArrayView<TSoftObjectPtr<>> Objects, const IPropertySourceProcessor& Source, const IReplicationStreamModel& Model, FEnumerateProperties Callback);
 }
\ No newline at end of file
diff --git a/Engine/Plugins/Developer/Concert/ConcertUI/ConcertSharedSlate/Source/ConcertSharedSlate/Private/Replication/Editor/View/ObjectEditor/SBaseReplicationStreamEditor.cpp b/Engine/Plugins/Developer/Concert/ConcertUI/ConcertSharedSlate/Source/ConcertSharedSlate/Private/Replication/Editor/View/ObjectEditor/SBaseReplicationStreamEditor.cpp
--- a/Engine/Plugins/Developer/Concert/ConcertUI/ConcertSharedSlate/Source/ConcertSharedSlate/Private/Replication/Editor/View/ObjectEditor/SBaseReplicationStreamEditor.cpp
+++ b/Engine/Plugins/Developer/Concert/ConcertUI/ConcertSharedSlate/Source/ConcertSharedSlate/Private/Replication/Editor/View/ObjectEditor/SBaseReplicationStreamEditor.cpp
@@ -7,7 +7,7 @@
 #include "Replication/Editor/Utils/DisplayUtils.h"
 #include "Replication/Editor/Model/IEditableReplicationStreamModel.h"
 #include "Replication/Editor/Model/Object/IObjectHierarchyModel.h"
-#include "Replication/Editor/Model/Property/IPropertySelectionSourceModel.h"
+#include "Replication/Editor/Model/Property/IPropertySourceProcessor.h"
 #include "Replication/Editor/Model/ObjectSource/IObjectSelectionSourceModel.h"
 #include "Replication/Editor/View/ObjectViewer/SReplicationStreamViewer.h"
 
@@ -27,7 +27,7 @@ namespace UE::ConcertSharedSlate
 		const FArguments& InArgs,
 		const TSharedRef<IEditableReplicationStreamModel>& InPropertiesModel,
 		const TSharedRef<IObjectSelectionSourceModel>& InObjectSelectionSource,
-		const TSharedRef<IPropertySelectionSourceModel>& InPropertySelectionSource)
+		const TSharedRef<IPropertySourceProcessor>& InPropertySelectionSource)
 	{
 		ObjectSelectionSource = InObjectSelectionSource;
 		PropertySelectionSource = InPropertySelectionSource;
diff --git a/Engine/Plugins/Developer/Concert/ConcertUI/ConcertSharedSlate/Source/ConcertSharedSlate/Private/Replication/Editor/View/ObjectEditor/SBaseReplicationStreamEditor.h b/Engine/Plugins/Developer/Concert/ConcertUI/ConcertSharedSlate/Source/ConcertSharedSlate/Private/Replication/Editor/View/ObjectEditor/SBaseReplicationStreamEditor.h
--- a/Engine/Plugins/Developer/Concert/ConcertUI/ConcertSharedSlate/Source/ConcertSharedSlate/Private/Replication/Editor/View/ObjectEditor/SBaseReplicationStreamEditor.h
+++ b/Engine/Plugins/Developer/Concert/ConcertUI/ConcertSharedSlate/Source/ConcertSharedSlate/Private/Replication/Editor/View/ObjectEditor/SBaseReplicationStreamEditor.h
@@ -23,7 +23,7 @@ namespace UE::ConcertSharedSlate
 	class IObjectNameModel;
 	class IObjectSelectionSourceModel;
 	class IPropertyAssignmentView;
-	class IPropertySelectionSourceModel;
+	class IPropertySourceProcessor;
 	class IObjectHierarchyModel;
 	class SReplicationStreamViewer;
 	
@@ -88,7 +88,7 @@ namespace UE::ConcertSharedSlate
 		void Construct(const FArguments& InArgs,
            const TSharedRef<IEditableReplicationStreamModel>& InPropertiesModel,
            const TSharedRef<IObjectSelectionSourceModel>& InObjectSelectionSource,
-           const TSharedRef<IPropertySelectionSourceModel>& InPropertySelectionSource
+           const TSharedRef<IPropertySourceProcessor>& InPropertySelectionSource
 		);
 		virtual ~SBaseReplicationStreamEditor() override;
 
@@ -112,7 +112,7 @@ namespace UE::ConcertSharedSlate
 		/** For deciding which objects can be added to EditablePropertiesModel. */
 		TSharedPtr<IObjectSelectionSourceModel> ObjectSelectionSource;
 		/** For deciding which properties can be added to EditablePropertiesModel. */
-		TSharedPtr<IPropertySelectionSourceModel> PropertySelectionSource;
+		TSharedPtr<IPropertySourceProcessor> PropertySelectionSource;
 
 		/** Optional. Determines whether all UI for changing the model should be disabled. */
 		TAttribute<bool> IsEditingEnabledAttribute;
diff --git a/Engine/Plugins/Developer/Concert/ConcertUI/ConcertSharedSlate/Source/ConcertSharedSlate/Private/Replication/Editor/View/Property/SMultiObjectAssignment.h b/Engine/Plugins/Developer/Concert/ConcertUI/ConcertSharedSlate/Source/ConcertSharedSlate/Private/Replication/Editor/View/Property/SMultiObjectAssignment.h
--- a/Engine/Plugins/Developer/Concert/ConcertUI/ConcertSharedSlate/Source/ConcertSharedSlate/Private/Replication/Editor/View/Property/SMultiObjectAssignment.h
+++ b/Engine/Plugins/Developer/Concert/ConcertUI/ConcertSharedSlate/Source/ConcertSharedSlate/Private/Replication/Editor/View/Property/SMultiObjectAssignment.h
@@ -15,7 +15,7 @@ struct TSoftObjectPtr;
 namespace UE::ConcertSharedSlate
 {
 	class IObjectHierarchyModel;
-	class IPropertySelectionSourceModel;
+	class IPropertySourceProcessor;
 	class IReplicationStreamModel;
 	
 	/** SMultiObjectAssignment shows the properties of the displayed object and all of its subobjects.*/
@@ -27,7 +27,7 @@ namespace UE::ConcertSharedSlate
 
 		SLATE_BEGIN_ARGS(SMultiObjectAssignment){}
 			/** Optional. If specified, displays the properties of this model instead of those assigned in the stream. */
-			SLATE_ARGUMENT(TSharedPtr<IPropertySelectionSourceModel>, PropertySource)
+			SLATE_ARGUMENT(TSharedPtr<IPropertySourceProcessor>, PropertySource)
 			
 			/** Optional. Gets components and subobjects of the displayed object. If unspecified, behaves exactly like SPerObjectAssignmentView. */
 			SLATE_ARGUMENT(TSharedPtr<IObjectHierarchyModel>, ObjectHierarchy)
@@ -54,7 +54,7 @@ namespace UE::ConcertSharedSlate
 		/** Used to get subobjects of selected objects. */
 		TSharedPtr<IObjectHierarchyModel> ObjectHierarchy;
 		/** Optional. If specified, displays the properties of this model instead of those assigned in the stream. */
-		TSharedPtr<IPropertySelectionSourceModel> OptionalPropertySource;
+		TSharedPtr<IPropertySourceProcessor> OptionalPropertySource;
 
 		/** Whether EChildRelationshipFlags::Subobject objects should be shown. */
 		bool bShouldShowSubobjects = false;
diff --git a/Engine/Plugins/Developer/Concert/ConcertUI/ConcertSharedSlate/Source/ConcertSharedSlate/Private/Replication/Editor/View/Property/SPerObjectPropertyAssignment.h b/Engine/Plugins/Developer/Concert/ConcertUI/ConcertSharedSlate/Source/ConcertSharedSlate/Private/Replication/Editor/View/Property/SPerObjectPropertyAssignment.h
--- a/Engine/Plugins/Developer/Concert/ConcertUI/ConcertSharedSlate/Source/ConcertSharedSlate/Private/Replication/Editor/View/Property/SPerObjectPropertyAssignment.h
+++ b/Engine/Plugins/Developer/Concert/ConcertUI/ConcertSharedSlate/Source/ConcertSharedSlate/Private/Replication/Editor/View/Property/SPerObjectPropertyAssignment.h
@@ -10,7 +10,7 @@
 
 namespace UE::ConcertSharedSlate
 {
-	class IPropertySelectionSourceModel;
+	class IPropertySourceProcessor;
 	class IReplicationStreamModel;
 	
 	/** SPerObjectPropertyAssignment shows only the properties of the displayed object, which is achieved by wrapping SPropertyTreeView. */
@@ -22,7 +22,7 @@ namespace UE::ConcertSharedSlate
 
 		SLATE_BEGIN_ARGS(SPerObjectPropertyAssignment){}
 			/** Optional. If specified, displays the properties of this model instead of those assigned in the stream. */
-			SLATE_ARGUMENT(TSharedPtr<IPropertySelectionSourceModel>, PropertySource)
+			SLATE_ARGUMENT(TSharedPtr<IPropertySourceProcessor>, PropertySource)
 		SLATE_END_ARGS()
 
 		void Construct(const FArguments& InArgs, TSharedRef<IPropertyTreeView> InTreeView);
@@ -42,7 +42,7 @@ namespace UE::ConcertSharedSlate
 		TArray<TSoftObjectPtr<>> PreviousSelectedObjects;
 		
 		/** Optional. If specified, displays the properties of this model instead of those assigned in the stream. */
-		TSharedPtr<IPropertySelectionSourceModel> OptionalPropertySource;
+		TSharedPtr<IPropertySourceProcessor> OptionalPropertySource;
 	};
 }
 
diff --git a/Engine/Plugins/Developer/Concert/ConcertUI/ConcertSharedSlate/Source/ConcertSharedSlate/Public/Model/Item/IItemSourceModel.h b/Engine/Plugins/Developer/Concert/ConcertUI/ConcertSharedSlate/Source/ConcertSharedSlate/Public/Model/Item/IItemSourceModel.h
--- a/Engine/Plugins/Developer/Concert/ConcertUI/ConcertSharedSlate/Source/ConcertSharedSlate/Public/Model/Item/IItemSourceModel.h
+++ b/Engine/Plugins/Developer/Concert/ConcertUI/ConcertSharedSlate/Source/ConcertSharedSlate/Public/Model/Item/IItemSourceModel.h
@@ -2,11 +2,13 @@
 
 #pragma once
 
+#include "Misc/EBreakBehavior.h"
+
 #include "Containers/Array.h"
+#include "Internationalization/Text.h"
 #include "Templates/Function.h"
 #include "Templates/SharedPointer.h"
 #include "Textures/SlateIcon.h"
-#include "Misc/EBreakBehavior.h"
 
 namespace UE::ConcertSharedSlate
 {
diff --git a/Engine/Plugins/Developer/Concert/ConcertUI/ConcertSharedSlate/Source/ConcertSharedSlate/Public/Model/Item/SourceSelectionCategory.h b/Engine/Plugins/Developer/Concert/ConcertUI/ConcertSharedSlate/Source/ConcertSharedSlate/Public/Model/Item/SourceSelectionCategory.h
--- a/Engine/Plugins/Developer/Concert/ConcertUI/ConcertSharedSlate/Source/ConcertSharedSlate/Public/Model/Item/SourceSelectionCategory.h
+++ b/Engine/Plugins/Developer/Concert/ConcertUI/ConcertSharedSlate/Source/ConcertSharedSlate/Public/Model/Item/SourceSelectionCategory.h
@@ -3,6 +3,8 @@
 #pragma once
 
 #include "IItemSourceModel.h"
+
+#include "Containers/Array.h"
 #include "Misc/Attribute.h"
 
 namespace UE::ConcertSharedSlate
diff --git a/Engine/Plugins/Developer/Concert/ConcertUI/ConcertSharedSlate/Source/ConcertSharedSlate/Public/Replication/Editor/Model/Property/IPropertySelectionSourceModel.h b/Engine/Plugins/Developer/Concert/ConcertUI/ConcertSharedSlate/Source/ConcertSharedSlate/Public/Replication/Editor/Model/Property/IPropertySelectionSourceModel.h
--- a/Engine/Plugins/Developer/Concert/ConcertUI/ConcertSharedSlate/Source/ConcertSharedSlate/Public/Replication/Editor/Model/Property/IPropertySelectionSourceModel.h
+++ b/Engine/Plugins/Developer/Concert/ConcertUI/ConcertSharedSlate/Source/ConcertSharedSlate/Public/Replication/Editor/Model/Property/IPropertySelectionSourceModel.h
@@ -5,39 +5,25 @@
 #include "IPropertySourceModel.h"
 #include "Model/Item/SourceSelectionCategory.h"
 
+#include "Misc/CoreMiscDefines.h"
 #include "Templates/SharedPointer.h"
 #include "UObject/SoftObjectPath.h"
 #include "UObject/SoftObjectPtr.h"
 
 namespace UE::ConcertSharedSlate
 {
-	using FPropertySourceCategory = ConcertSharedSlate::TSourceSelectionCategory<FSelectablePropertyInfo>;
-
-	/** Arguments for getting properties associated with an object / class. */
-	struct FPropertySourceContext
-	{
-		/** The object for which the properties are supposed to be displayed */
-		TSoftObjectPtr<> Object;
-		/** The class of Object */
-		FSoftClassPath Class;
-
-		FPropertySourceContext(const TSoftObjectPtr<>& Object, const FSoftClassPath& Class)
-			: Object(Object)
-			, Class(Class)
-		{}
-	};
+	using FPropertySourceCategory UE_DEPRECATED(5.5, "No longer in use.") = ConcertSharedSlate::TSourceSelectionCategory<FSelectablePropertyInfo>;
 	
 	/** Decides which properties can be added to a IReplicationStreamModel. */
-	class IPropertySelectionSourceModel : public TSharedFromThis<IPropertySelectionSourceModel>
+	class UE_DEPRECATED(5.5, "Use IPropertySourceProcessor instead.") IPropertySelectionSourceModel
+		PRAGMA_DISABLE_DEPRECATION_WARNINGS
+		: public TSharedFromThis<IPropertySelectionSourceModel>
+		PRAGMA_ENABLE_DEPRECATION_WARNINGS
 	{
 	public:
 
 		/** Gets the single source determining which properties can be selected. */
-		UE_DEPRECATED(5.5, "Use the version that takes FPropertySourceContext instead.")
-		virtual TSharedRef<IPropertySourceModel> GetPropertySource(const FSoftClassPath& Class) const { return GetPropertySource(FPropertySourceContext({}, Class)); }
-
-		/** @return Gets an object that will tell you the properties that are associated with the passed in object / class context. */
-		virtual TSharedRef<IPropertySourceModel> GetPropertySource(const FPropertySourceContext& Context) const = 0;
+		virtual TSharedRef<IPropertySourceModel> GetPropertySource(const FSoftClassPath& Class) const = 0;
 		
 		virtual ~IPropertySelectionSourceModel() = default;
 	};
diff --git a/Engine/Plugins/Developer/Concert/ConcertUI/ConcertSharedSlate/Source/ConcertSharedSlate/Public/Replication/Editor/Model/Property/IPropertySourceModel.h b/Engine/Plugins/Developer/Concert/ConcertUI/ConcertSharedSlate/Source/ConcertSharedSlate/Public/Replication/Editor/Model/Property/IPropertySourceModel.h
--- a/Engine/Plugins/Developer/Concert/ConcertUI/ConcertSharedSlate/Source/ConcertSharedSlate/Public/Replication/Editor/Model/Property/IPropertySourceModel.h
+++ b/Engine/Plugins/Developer/Concert/ConcertUI/ConcertSharedSlate/Source/ConcertSharedSlate/Public/Replication/Editor/Model/Property/IPropertySourceModel.h
@@ -15,7 +15,7 @@ namespace UE::ConcertSharedSlate
 	
 	/**
 	 * A specific object source, e.g. like actors, components from an actor (right-click), etc.
-	 * @see IPropertySelectionSourceModel.
+	 * @see IPropertySelectionProcessor.
 	 */
 	using IPropertySourceModel = ConcertSharedSlate::IItemSourceModel<FSelectablePropertyInfo>;
 }
diff --git a/Engine/Plugins/Developer/Concert/ConcertUI/ConcertSharedSlate/Source/ConcertSharedSlate/Public/Replication/Editor/Model/Property/IPropertySourceProcessor.h b/Engine/Plugins/Developer/Concert/ConcertUI/ConcertSharedSlate/Source/ConcertSharedSlate/Public/Replication/Editor/Model/Property/IPropertySourceProcessor.h
new file mode 100644
--- /dev/null
+++ b/Engine/Plugins/Developer/Concert/ConcertUI/ConcertSharedSlate/Source/ConcertSharedSlate/Public/Replication/Editor/Model/Property/IPropertySourceProcessor.h
@@ -0,0 +1,42 @@
+﻿// Copyright Epic Games, Inc. All Rights Reserved.
+
+#pragma once
+
+#include "Replication/Editor/Model/Property/IPropertySourceModel.h"
+
+#include "Templates/FunctionFwd.h"
+#include "UObject/SoftObjectPath.h"
+#include "UObject/SoftObjectPtr.h"
+
+namespace UE::ConcertSharedSlate
+{
+	/** Arguments for getting properties associated with an object / class. */
+	struct FPropertySourceContext
+	{
+		/** The object for which the properties are supposed to be displayed */
+		TSoftObjectPtr<> Object;
+		/** The class of Object */
+		FSoftClassPath Class;
+
+		FPropertySourceContext(const TSoftObjectPtr<>& Object, const FSoftClassPath& Class)
+			: Object(Object)
+			, Class(Class)
+		{}
+	};
+	
+	/**
+	 * Determines the properties that should be displayed for an Object / Class.
+	 * 
+	 * The most simple implementation is to iterate the UClass properties but there can be more advanced implementations, such as only returning
+	 * properties from a user defined list.
+	 */
+	class IPropertySourceProcessor
+	{
+	public:
+		
+		/** Allows you to get the properties that are associated with the passed in object / class context. You are not allowed to keep a reference to Processor. */
+		virtual void ProcessPropertySource(const FPropertySourceContext& Context, TFunctionRef<void(const IPropertySourceModel& Model)> Processor) const = 0;
+		
+		virtual ~IPropertySourceProcessor() = default;
+	};
+}
diff --git a/Engine/Plugins/Developer/Concert/ConcertUI/ConcertSharedSlate/Source/ConcertSharedSlate/Public/Replication/PropertyAssignmentViewFactory.h b/Engine/Plugins/Developer/Concert/ConcertUI/ConcertSharedSlate/Source/ConcertSharedSlate/Public/Replication/PropertyAssignmentViewFactory.h
--- a/Engine/Plugins/Developer/Concert/ConcertUI/ConcertSharedSlate/Source/ConcertSharedSlate/Public/Replication/PropertyAssignmentViewFactory.h
+++ b/Engine/Plugins/Developer/Concert/ConcertUI/ConcertSharedSlate/Source/ConcertSharedSlate/Public/Replication/PropertyAssignmentViewFactory.h
@@ -16,7 +16,7 @@ namespace UE::ConcertSharedSlate
 	class IMultiObjectPropertyAssignmentView;
 	class FPropertyData;
 	class IPropertyAssignmentView;
-	class IPropertySelectionSourceModel;
+	class IPropertySourceProcessor;
 	class IPropertyTreeView;
 	
 	DECLARE_DELEGATE_RetVal_OneParam(EFilterResult, FFilterPropertyData, const FPropertyData&);
@@ -31,7 +31,7 @@ namespace UE::ConcertSharedSlate
 		 * If specified, the view will display all properties reported by the model (useful for editor UI which edits streams, not useful for server where property info is not available).
 		 * If unspecified, only display the properties assigned to the object in the stream.
 		 */
-		TSharedPtr<IPropertySelectionSourceModel> PropertySource;
+		TSharedPtr<IPropertySourceProcessor> PropertySource;
 	};
 	
 	/**
@@ -59,7 +59,7 @@ namespace UE::ConcertSharedSlate
 		 * If specified, the view will display all properties reported by the model (useful for editor UI which edits streams, not useful for server where property info is not available).
 		 * If unspecified, only display the properties assigned to the object in the stream.
 		 */
-		TSharedPtr<IPropertySelectionSourceModel> PropertySource;
+		TSharedPtr<IPropertySourceProcessor> PropertySource;
 	};
 
 	/**
diff --git a/Engine/Plugins/Developer/Concert/ConcertUI/ConcertSharedSlate/Source/ConcertSharedSlate/Public/Replication/ReplicationWidgetFactories.h b/Engine/Plugins/Developer/Concert/ConcertUI/ConcertSharedSlate/Source/ConcertSharedSlate/Public/Replication/ReplicationWidgetFactories.h
--- a/Engine/Plugins/Developer/Concert/ConcertUI/ConcertSharedSlate/Source/ConcertSharedSlate/Public/Replication/ReplicationWidgetFactories.h
+++ b/Engine/Plugins/Developer/Concert/ConcertUI/ConcertSharedSlate/Source/ConcertSharedSlate/Public/Replication/ReplicationWidgetFactories.h
@@ -36,7 +36,7 @@ namespace UE::ConcertSharedSlate
 	class IReplicationStreamViewer;
 	class IStreamExtender;
 	class IObjectHierarchyModel;
-	class IPropertySelectionSourceModel;
+	class IPropertySourceProcessor;
 	
 	/**
 	 * Creates a model that can be passed to CreateEditor.
@@ -113,7 +113,7 @@ namespace UE::ConcertSharedSlate
 		 * Required. Determines the properties that are displayed in the property list. 
 		 * @note The view will keep a strong reference to this.
 		 */
-		TSharedRef<IPropertySelectionSourceModel> PropertySource;
+		TSharedRef<IPropertySourceProcessor> PropertySource;
 		
 		/** Optional. Determines whether all UI for changing the model should be disabled. */
 		TAttribute<bool> IsEditingEnabled;
@@ -171,7 +171,7 @@ namespace UE::ConcertSharedSlate
 		 * Determines the properties that can be added to the property list.
 		 * @note The view will keep a strong reference to this.
 		 */
-		TSharedRef<IPropertySelectionSourceModel> PropertySource;
+		TSharedRef<IPropertySourceProcessor> PropertySource;
 
 		/** Optional. If set, the Add button should automatically assign the added object to stream returned by this callback. */
 		FGetAutoAssignTarget GetAutoAssignToStreamDelegate;
diff --git a/Engine/Plugins/Developer/Concert/ConcertUI/ConcertSharedSlate/Source/ConcertSharedSlate/Public/Replication/Utils/ReplicationWidgetDelegates.h b/Engine/Plugins/Developer/Concert/ConcertUI/ConcertSharedSlate/Source/ConcertSharedSlate/Public/Replication/Utils/ReplicationWidgetDelegates.h
--- a/Engine/Plugins/Developer/Concert/ConcertUI/ConcertSharedSlate/Source/ConcertSharedSlate/Public/Replication/Utils/ReplicationWidgetDelegates.h
+++ b/Engine/Plugins/Developer/Concert/ConcertUI/ConcertSharedSlate/Source/ConcertSharedSlate/Public/Replication/Utils/ReplicationWidgetDelegates.h
@@ -17,7 +17,7 @@ struct FSoftObjectPath;
 namespace UE::ConcertSharedSlate
 {
 	class FPropertyData;
-	class IPropertySelectionSourceModel;
+	class IPropertySourceProcessor;
 	class IReplicationStreamModel;
 
 	/** A predicate for determining Left < Right. */
