Deprecate IItemSourceModel::GetNumSelectableItems

diff --git a/Engine/Plugins/Developer/Concert/ConcertUI/ConcertSharedSlate/Source/ConcertSharedSlate/Public/Model/Item/IItemSourceModel.h b/Engine/Plugins/Developer/Concert/ConcertUI/ConcertSharedSlate/Source/ConcertSharedSlate/Public/Model/Item/IItemSourceModel.h
--- a/Engine/Plugins/Developer/Concert/ConcertUI/ConcertSharedSlate/Source/ConcertSharedSlate/Public/Model/Item/IItemSourceModel.h
+++ b/Engine/Plugins/Developer/Concert/ConcertUI/ConcertSharedSlate/Source/ConcertSharedSlate/Public/Model/Item/IItemSourceModel.h
@@ -60,16 +60,24 @@ namespace UE::ConcertSharedSlate
 		virtual FSourceDisplayInfo GetDisplayInfo() const = 0;
 
 		/** @return The number of selectable objects. */
-		virtual uint32 GetNumSelectableItems() const = 0;
+		UE_DEPRECATED(5.5, "No longer in use")
+		virtual uint32 GetNumSelectableItems() const { return 0; }
 		
 		/** Enumerates ALL objects from this source. This can be a LOT of objects (e.g. 2000 actors). Intended to be placed in a searchable drop-down menu. */
 		virtual void EnumerateSelectableItems(TFunctionRef<EBreakBehavior(const TItemType& SelectableOption)> Delegate) const = 0;
 
+		/** @return Whether there are any options. */
+		virtual bool HasOptions() const
+		{
+			bool bHasAnItem = false;
+			EnumerateSelectableItems([&bHasAnItem](const auto&){ bHasAnItem = true; return EBreakBehavior::Break; });
+			return bHasAnItem;
+		}
+		
 		/** Util that converts EnumerateSelectableItems into an array. */
 		TArray<TItemType> GetSelectableItems()
 		{
 			TArray<TItemType> Result;
-			Result.Reserve(GetNumSelectableItems());
 			EnumerateSelectableItems([&Result](const TItemType& SelectableOption)
 			{
 				Result.Add(SelectableOption);
diff --git a/Engine/Plugins/Developer/Concert/ConcertUI/ConcertSharedSlate/Source/ConcertSharedSlate/Public/Model/Item/SourceModelBuilders.h b/Engine/Plugins/Developer/Concert/ConcertUI/ConcertSharedSlate/Source/ConcertSharedSlate/Public/Model/Item/SourceModelBuilders.h
--- a/Engine/Plugins/Developer/Concert/ConcertUI/ConcertSharedSlate/Source/ConcertSharedSlate/Public/Model/Item/SourceModelBuilders.h
+++ b/Engine/Plugins/Developer/Concert/ConcertUI/ConcertSharedSlate/Source/ConcertSharedSlate/Public/Model/Item/SourceModelBuilders.h
@@ -242,7 +242,7 @@ namespace UE::ConcertSharedSlate
 				AddOnClickParams.UserInterfaceActionType = EUserInterfaceActionType::Button; 
 				AddOnClickParams.DirectActions = {
 					FExecuteAction::CreateLambda([Option, Args = MoveTemp(Args)](){ Args.OnItemsSelected.Execute(Option->GetSelectableItems()); }),
-					FCanExecuteAction::CreateLambda([Option](){ return Option->GetNumSelectableItems() > 0; })
+					FCanExecuteAction::CreateLambda([Option](){ return Option->HasOptions(); })
 				};
 				
 				MenuBuilder.AddMenuEntry(AddOnClickParams);
