Adjust IPropertyAssignmentView API to expose the objects displayed by it

diff --git a/Engine/Plugins/Developer/Concert/ConcertUI/ConcertSharedSlate/Source/ConcertSharedSlate/Private/Replication/Editor/View/Property/SMultiObjectAssignment.cpp b/Engine/Plugins/Developer/Concert/ConcertUI/ConcertSharedSlate/Source/ConcertSharedSlate/Private/Replication/Editor/View/Property/SMultiObjectAssignment.cpp
--- a/Engine/Plugins/Developer/Concert/ConcertUI/ConcertSharedSlate/Source/ConcertSharedSlate/Private/Replication/Editor/View/Property/SMultiObjectAssignment.cpp
+++ b/Engine/Plugins/Developer/Concert/ConcertUI/ConcertSharedSlate/Source/ConcertSharedSlate/Private/Replication/Editor/View/Property/SMultiObjectAssignment.cpp
@@ -115,12 +115,15 @@ namespace UE::ConcertSharedSlate
 	{
 		TArray<FPropertyAssignmentEntry> Entries;
 		auto[RootEntry, bRootSharesClass] = BuildAssignmentEntry(Objects, Model);
-		if (!ensureMsgf(bRootSharesClass, TEXT("Objects do not share the same class. Investigate invalid call.")))
+		if (RootEntry.PropertiesToDisplay.IsEmpty() || !ensureMsgf(bRootSharesClass, TEXT("Objects do not share the same class. Investigate invalid call.")))
 		{
-			PreviousSelectedObjects.Reset();
+			PreviousSelectedObjects = Objects;
 			TreeView->RefreshPropertyData({}, false);
+			OnObjectGroupsChangedDelegate.Broadcast();
+			return;
 		}
-		
+
+		DisplayedGroups = { FObjectGroup{ Objects } };
 		Entries.Emplace(MoveTemp(RootEntry));
 		if (ObjectHierarchy)
 		{
@@ -128,6 +131,8 @@ namespace UE::ConcertSharedSlate
 			const TMap<FSoftObjectPath, TArray<FSoftObjectPath>> HierarchyInfo = BuildHierarchy(Objects, *ObjectHierarchy, bShouldShowSubobjects);
 			ForEachCategory(Objects, HierarchyInfo, [this, &Model, &Entries](const TArray<TSoftObjectPtr<>>& ContextObjects)
 			{
+				DisplayedGroups.Add(FObjectGroup{ ContextObjects });
+				
 				auto[Entry, bShareClass] = BuildAssignmentEntry(ContextObjects, Model);
 				if (bShareClass)
 				{
@@ -136,15 +141,22 @@ namespace UE::ConcertSharedSlate
 			});
 		}
 		
+		const bool bObjectsHaveChanged = PreviousSelectedObjects == Objects;
+		PreviousSelectedObjects = Objects;
+		
 		// If the objects have changed, the classes may share properties.
 		// In that case, below we'd reuse the item pointer, which would cause the tree view to re-use the old row widgets.
 		// However, we must regenerate all column widgets since they may be referencing the object the row was originally built for. So they'd display the state of the previous object still!
 		// Example: Assign property combo-box in Multi-User All Clients view displays who has the property assigned.
 		// Note: If the objects did not change, we definitely want to reuse item pointers since otherwise the user row selection is reset.
-		const bool bCanReusePropertyData = PreviousSelectedObjects == Objects;
-		PreviousSelectedObjects = Objects;
-			
+		const bool bCanReusePropertyData = bObjectsHaveChanged;
+		
 		TreeView->RefreshPropertyData(Entries, bCanReusePropertyData);
+
+		if (!bObjectsHaveChanged)
+		{
+			OnObjectGroupsChangedDelegate.Broadcast();
+		}
 	}
 
 	void SMultiObjectAssignment::SetShouldShowSubobjects(bool bShowSubobjects)
diff --git a/Engine/Plugins/Developer/Concert/ConcertUI/ConcertSharedSlate/Source/ConcertSharedSlate/Private/Replication/Editor/View/Property/SMultiObjectAssignment.h b/Engine/Plugins/Developer/Concert/ConcertUI/ConcertSharedSlate/Source/ConcertSharedSlate/Private/Replication/Editor/View/Property/SMultiObjectAssignment.h
--- a/Engine/Plugins/Developer/Concert/ConcertUI/ConcertSharedSlate/Source/ConcertSharedSlate/Private/Replication/Editor/View/Property/SMultiObjectAssignment.h
+++ b/Engine/Plugins/Developer/Concert/ConcertUI/ConcertSharedSlate/Source/ConcertSharedSlate/Private/Replication/Editor/View/Property/SMultiObjectAssignment.h
@@ -40,9 +40,14 @@ namespace UE::ConcertSharedSlate
 		virtual void RequestRefilter() const override { return TreeView->RequestRefilter(); }
 		virtual void RequestResortForColumn(const FName& ColumnId) override { return TreeView->RequestRefilter(); }
 		virtual TSharedRef<SWidget> GetWidget() override { return SharedThis(this); }
+		virtual TArray<FObjectGroup> GetDisplayedGroups() const override { return DisplayedGroups; }
+		virtual FOnSelectionChanged& OnObjectGroupsChanged() override { return OnObjectGroupsChangedDelegate; }
+		//~ End IPropertyAssignmentView Interface
+		
+		//~ End IMultiObjectPropertyAssignmentView Interface
 		virtual void SetShouldShowSubobjects(bool bShowSubobjects) override;
 		virtual bool GetShouldShowSubobjects() const override { return bShouldShowSubobjects;}
-		//~ End IPropertyAssignmentView Interface
+		//~ End IMultiObjectPropertyAssignmentView Interface
 
 	private:
 
@@ -50,6 +55,8 @@ namespace UE::ConcertSharedSlate
 		TSharedPtr<IPropertyTreeView> TreeView;
 		/** Used to determine whether to rebuild the entire property data. */
 		TArray<TSoftObjectPtr<>> PreviousSelectedObjects;
+		/** Cached value for GetDisplayedGroups. */
+		TArray<FObjectGroup> DisplayedGroups;
 
 		/** Used to get subobjects of selected objects. */
 		TSharedPtr<IObjectHierarchyModel> ObjectHierarchy;
@@ -58,6 +65,9 @@ namespace UE::ConcertSharedSlate
 
 		/** Whether EChildRelationshipFlags::Subobject objects should be shown. */
 		bool bShouldShowSubobjects = false;
+		
+		/** Broadcasts when the result of GetDisplayedGroups() has changed. */
+		FOnSelectionChanged OnObjectGroupsChangedDelegate;
 
 		struct FBuildAssignmentEntryResult
 		{
diff --git a/Engine/Plugins/Developer/Concert/ConcertUI/ConcertSharedSlate/Source/ConcertSharedSlate/Private/Replication/Editor/View/Property/SPerObjectPropertyAssignment.cpp b/Engine/Plugins/Developer/Concert/ConcertUI/ConcertSharedSlate/Source/ConcertSharedSlate/Private/Replication/Editor/View/Property/SPerObjectPropertyAssignment.cpp
--- a/Engine/Plugins/Developer/Concert/ConcertUI/ConcertSharedSlate/Source/ConcertSharedSlate/Private/Replication/Editor/View/Property/SPerObjectPropertyAssignment.cpp
+++ b/Engine/Plugins/Developer/Concert/ConcertUI/ConcertSharedSlate/Source/ConcertSharedSlate/Private/Replication/Editor/View/Property/SPerObjectPropertyAssignment.cpp
@@ -45,22 +45,29 @@ namespace UE::ConcertSharedSlate
 				return EBreakBehavior::Continue;
 			});
 		
-		if (bHasClassPath)
+		if (bHasClassPath && !Properties.IsEmpty())
 		{
+			const bool bObjectsHaveChanged = PreviousSelectedObjects == Objects;
+			PreviousSelectedObjects = Objects;
+			
 			// If the objects have changed, the classes may share properties.
 			// In that case, below we'd reuse the item pointer, which would cause the tree view to re-use the old row widgets.
 			// However, we must regenerate all column widgets since they may be referencing the object the row was originally built for. So they'd display the state of the previous object still!
 			// Example: Assign property combo-box in Multi-User All Clients view displays who has the property assigned.
 			// Note: If the objects did not change, we definitely want to reuse item pointers since otherwise the user row selection is reset.
-			const bool bCanReusePropertyData = PreviousSelectedObjects == Objects;
-			PreviousSelectedObjects = Objects;
-			
+			const bool bCanReusePropertyData = bObjectsHaveChanged;
 			TreeView->RefreshPropertyData({ AssignmentEntry }, bCanReusePropertyData);
+
+			if (!bObjectsHaveChanged)
+			{
+				OnObjectGroupsChangedDelegate.Broadcast();
+			}
 		}
 		else
 		{
-			PreviousSelectedObjects.Reset();
+			PreviousSelectedObjects = Objects;
 			TreeView->RefreshPropertyData({}, false);
+			OnObjectGroupsChangedDelegate.Broadcast();
 		}
 	}
 }
diff --git a/Engine/Plugins/Developer/Concert/ConcertUI/ConcertSharedSlate/Source/ConcertSharedSlate/Private/Replication/Editor/View/Property/SPerObjectPropertyAssignment.h b/Engine/Plugins/Developer/Concert/ConcertUI/ConcertSharedSlate/Source/ConcertSharedSlate/Private/Replication/Editor/View/Property/SPerObjectPropertyAssignment.h
--- a/Engine/Plugins/Developer/Concert/ConcertUI/ConcertSharedSlate/Source/ConcertSharedSlate/Private/Replication/Editor/View/Property/SPerObjectPropertyAssignment.h
+++ b/Engine/Plugins/Developer/Concert/ConcertUI/ConcertSharedSlate/Source/ConcertSharedSlate/Private/Replication/Editor/View/Property/SPerObjectPropertyAssignment.h
@@ -32,6 +32,8 @@ namespace UE::ConcertSharedSlate
 		virtual void RequestRefilter() const override { return TreeView->RequestRefilter(); }
 		virtual void RequestResortForColumn(const FName& ColumnId) override { return TreeView->RequestRefilter(); }
 		virtual TSharedRef<SWidget> GetWidget() override { return SharedThis(this); }
+		virtual TArray<FObjectGroup> GetDisplayedGroups() const override { return { FObjectGroup{ PreviousSelectedObjects } }; }
+		virtual FOnSelectionChanged& OnObjectGroupsChanged() override { return OnObjectGroupsChangedDelegate; }
 		//~ End IPropertyAssignmentView Interface
 
 	private:
@@ -43,6 +45,9 @@ namespace UE::ConcertSharedSlate
 		
 		/** Optional. If specified, displays the properties of this model instead of those assigned in the stream. */
 		TSharedPtr<IPropertySourceProcessor> OptionalPropertySource;
+		
+		/** Broadcasts when the result of GetDisplayedGroups() has changed. */
+		FOnSelectionChanged OnObjectGroupsChangedDelegate;
 	};
 }
 
diff --git a/Engine/Plugins/Developer/Concert/ConcertUI/ConcertSharedSlate/Source/ConcertSharedSlate/Public/Replication/Editor/View/IMultiObjectPropertyAssignmentView.h b/Engine/Plugins/Developer/Concert/ConcertUI/ConcertSharedSlate/Source/ConcertSharedSlate/Public/Replication/Editor/View/IMultiObjectPropertyAssignmentView.h
--- a/Engine/Plugins/Developer/Concert/ConcertUI/ConcertSharedSlate/Source/ConcertSharedSlate/Public/Replication/Editor/View/IMultiObjectPropertyAssignmentView.h
+++ b/Engine/Plugins/Developer/Concert/ConcertUI/ConcertSharedSlate/Source/ConcertSharedSlate/Public/Replication/Editor/View/IMultiObjectPropertyAssignmentView.h
@@ -15,7 +15,5 @@ namespace UE::ConcertSharedSlate
 		virtual void SetShouldShowSubobjects(bool bShowSubobjects) = 0;
 		/** @return Whether subobjects should be shown */
 		virtual bool GetShouldShowSubobjects() const = 0;
-
-		virtual ~IMultiObjectPropertyAssignmentView() = default;
 	};
 }
diff --git a/Engine/Plugins/Developer/Concert/ConcertUI/ConcertSharedSlate/Source/ConcertSharedSlate/Public/Replication/Editor/View/IPropertyAssignmentView.h b/Engine/Plugins/Developer/Concert/ConcertUI/ConcertSharedSlate/Source/ConcertSharedSlate/Public/Replication/Editor/View/IPropertyAssignmentView.h
--- a/Engine/Plugins/Developer/Concert/ConcertUI/ConcertSharedSlate/Source/ConcertSharedSlate/Public/Replication/Editor/View/IPropertyAssignmentView.h
+++ b/Engine/Plugins/Developer/Concert/ConcertUI/ConcertSharedSlate/Source/ConcertSharedSlate/Public/Replication/Editor/View/IPropertyAssignmentView.h
@@ -13,6 +13,19 @@ class UObject;
 namespace UE::ConcertSharedSlate
 {
 	class IReplicationStreamModel;
+
+	/**
+	 * An object group is a bunch of related objects.
+	 * This relates to multi-editing.
+	 *
+	 * Example: You click 2 ACineCameraActor in the IReplicationStreamViewer:
+	 * - A group is the 2 actors you clicked
+	 * - A group is the two cine camera components of the actor
+	 */
+	struct FObjectGroup
+	{
+		TArray<TSoftObjectPtr<>> Group;
+	};
 	
 	/**
 	 * A replication assignment view displays an object's properties.
@@ -22,8 +35,9 @@ namespace UE::ConcertSharedSlate
 	 * - Object tree view: User can click on an object.
 	 * - IPropertyAssignmentView: The clicked object's properties are displayed (the "root objects").
 	 *
-	 * Right now there is only 1 implementation: SPerObjectPropertyAssignment, which displays the root object's properties only.
-	 * We could also add an implementation that also displays the subobject properties, similar how the editor's property matrix does it. 
+	 * Right now there are 2 implementation:
+	 *	- SPerObjectPropertyAssignment, which displays the root object's properties only
+	 *	- SMultiObjectAssignment, which displays the root object and its subobjects.
 	 * Generally, you can imagine the view as a tree view which has columns that can be injected (e.g. @see CreateBaseStreamEditor).
 	 */
 	class CONCERTSHAREDSLATE_API IPropertyAssignmentView
@@ -50,6 +64,13 @@ namespace UE::ConcertSharedSlate
 		/** Gets the tree view's widget */
 		virtual TSharedRef<SWidget> GetWidget() = 0;
 
+		/** @return The groups of objects being displayed in this view right now. */
+		virtual TArray<FObjectGroup> GetDisplayedGroups() const = 0;
+		
+		DECLARE_MULTICAST_DELEGATE(FOnSelectionChanged);
+		/** @return Event that broadcasts when the result of GetDisplayedGroups() has changed. */
+		virtual FOnSelectionChanged& OnObjectGroupsChanged() = 0;
+
 		virtual ~IPropertyAssignmentView() = default;
 	};
 }
