Add ESourceType::ShowAsToggleButtonList for the IItemModels, which allows the combo button entries to have checkboxes

diff --git a/Engine/Plugins/Developer/Concert/ConcertUI/ConcertSharedSlate/Source/ConcertSharedSlate/Public/Model/Item/IItemSourceModel.h b/Engine/Plugins/Developer/Concert/ConcertUI/ConcertSharedSlate/Source/ConcertSharedSlate/Public/Model/Item/IItemSourceModel.h
--- a/Engine/Plugins/Developer/Concert/ConcertUI/ConcertSharedSlate/Source/ConcertSharedSlate/Public/Model/Item/IItemSourceModel.h
+++ b/Engine/Plugins/Developer/Concert/ConcertUI/ConcertSharedSlate/Source/ConcertSharedSlate/Public/Model/Item/IItemSourceModel.h
@@ -14,10 +14,26 @@ namespace UE::ConcertSharedSlate
 {
 	enum class ESourceType : uint8
 	{
-		/** Display EnumerateSelectableItems in a drop-down or menu list. Example: User wants to search from a list of actors. */
+		/**
+		 * Display EnumerateSelectableItems in a drop-down or menu list. Each entry is a button that closes the menu.
+		 * Example: User wants to search from a list of actors.
+		 */
 		ShowAsList,
-		/** When this option is clicked, instantly add all of EnumerateSelectableItems. Example: User wants to add all actors selected in the world outliner. */
-		AddOnClick
+		/**
+		 * When this option is clicked, instantly add all of EnumerateSelectableItems.
+		 * Example: User wants to add all actors selected in the world outliner.
+		 */
+		AddOnClick,
+		/**
+		 * Display EnumerateSelectableItems in a drop-down or menu list.
+		 *
+		 * Each entry contains a checkbox:
+		 * - FSourceModelBuilders::OnItemsSelected is called when the checkbox is toggled (since you provide IsItemSelected, you can determine whether to add or remove the items).
+		 * - FSourceModelBuilders::IsItemSelected determines whether the checkbox is checked
+		 * 
+		 * Example: User wants to search from a list of actors.
+		 */
+		ShowAsToggleButtonList,
 	};
 
 	struct FBaseDisplayInfo
diff --git a/Engine/Plugins/Developer/Concert/ConcertUI/ConcertSharedSlate/Source/ConcertSharedSlate/Public/Model/Item/SourceModelBuilders.h b/Engine/Plugins/Developer/Concert/ConcertUI/ConcertSharedSlate/Source/ConcertSharedSlate/Public/Model/Item/SourceModelBuilders.h
--- a/Engine/Plugins/Developer/Concert/ConcertUI/ConcertSharedSlate/Source/ConcertSharedSlate/Public/Model/Item/SourceModelBuilders.h
+++ b/Engine/Plugins/Developer/Concert/ConcertUI/ConcertSharedSlate/Source/ConcertSharedSlate/Public/Model/Item/SourceModelBuilders.h
@@ -154,19 +154,23 @@ namespace UE::ConcertSharedSlate
 	{
 		const FSourceDisplayInfo DisplayInfo = Source->GetDisplayInfo();
 		const FSlateBrush* IconBrush = DisplayInfo.Icon.IsSet() ? DisplayInfo.Icon.GetOptionalIcon() : FAppStyle::Get().GetBrush("Icons.Plus");
+
+		const auto GetToolTipText = [ToolTipText = DisplayInfo.ToolTip, DisabledToolTipAttribute = Args.DisabledToolTipAttribute, IsEnabled = Args.IsEnabledAttribute]()
+		{
+			const bool bIsEnabled = !(IsEnabled.IsBound() || IsEnabled.IsSet()) || IsEnabled.Get();
+			const bool bHasDisabledText = DisabledToolTipAttribute.IsBound() || DisabledToolTipAttribute.IsSet();
+			return bIsEnabled
+				? ToolTipText
+				: bHasDisabledText ? DisabledToolTipAttribute.Get() : FText::GetEmpty();
+		};
+		
 		switch (DisplayInfo.SourceType)
 		{
+		case ESourceType::ShowAsToggleButtonList: [[fallthrough]];
 		case ESourceType::ShowAsList:
 			return SNew(SPositiveActionButton)
 				.Text(DisplayInfo.Label)
-				.ToolTipText_Lambda([ToolTipText = DisplayInfo.ToolTip, DisabledToolTipAttribute = Args.DisabledToolTipAttribute, IsEnabled = Args.IsEnabledAttribute]()
-				{
-					const bool bIsEnabled = !(IsEnabled.IsBound() || IsEnabled.IsSet()) || IsEnabled.Get();
-					const bool bHasDisabledText = DisabledToolTipAttribute.IsBound() || DisabledToolTipAttribute.IsSet();
-					return bIsEnabled
-						? ToolTipText
-						: bHasDisabledText ? DisabledToolTipAttribute.Get() : FText::GetEmpty();
-				})
+				.ToolTipText_Lambda(GetToolTipText)
 				.Icon(IconBrush)
 				.IsEnabled_Lambda([IsEnabled = Args.IsEnabledAttribute](){ return !(IsEnabled.IsBound() || IsEnabled.IsSet()) || IsEnabled.Get(); })
 				.OnGetMenuContent_Lambda([Source = MoveTemp(Source), Args = MoveTemp(Args)]()
@@ -178,14 +182,7 @@ namespace UE::ConcertSharedSlate
 		case ESourceType::AddOnClick: 
 			return SNew(SPositiveActionButton)
 				.Text(DisplayInfo.Label)
-				.ToolTipText_Lambda([ToolTipText = DisplayInfo.ToolTip, DisabledToolTipAttribute = Args.DisabledToolTipAttribute, IsEnabled = Args.IsEnabledAttribute]()
-				{
-					const bool bIsEnabled = !(IsEnabled.IsBound() || IsEnabled.IsSet()) || IsEnabled.Get();
-					const bool bHasDisabledText = DisabledToolTipAttribute.IsBound() || DisabledToolTipAttribute.IsSet();
-					return bIsEnabled
-						? ToolTipText
-						: bHasDisabledText ? DisabledToolTipAttribute.Get() : FText::GetEmpty();
-				})
+				.ToolTipText_Lambda(GetToolTipText)
 				.Icon(IconBrush)
 				.IsEnabled_Lambda([IsEnabled = Args.IsEnabledAttribute](){ return !(IsEnabled.IsBound() || IsEnabled.IsSet()) || IsEnabled.Get(); })
 				.OnClicked_Lambda([Source = MoveTemp(Source), AddObjectsDelegate = Args.OnItemsSelected]()
@@ -206,7 +203,8 @@ namespace UE::ConcertSharedSlate
 		const FSourceDisplayInfo DisplayInfo = Option->GetDisplayInfo();
 		switch (DisplayInfo.SourceType)
 		{
-		case ESourceType::ShowAsList:
+		case ESourceType::ShowAsList: [[fallthrough]];
+		case ESourceType::ShowAsToggleButtonList:
 			{
 				if (EnumHasAnyFlags(Args.Flags, EItemPickerFlags::DisplayOptionListInline))
 				{
@@ -217,6 +215,7 @@ namespace UE::ConcertSharedSlate
 					FMenuEntryParams ShowAsListParams;
 					ShowAsListParams.LabelOverride = DisplayInfo.Label;
 					ShowAsListParams.ToolTipOverride = DisplayInfo.ToolTip;
+					ShowAsListParams.IconOverride = DisplayInfo.Icon;
 					ShowAsListParams.bIsSubMenu = true;
 					// Dummy is needed to avoid assert
 					ShowAsListParams.DirectActions = {
@@ -239,6 +238,7 @@ namespace UE::ConcertSharedSlate
 				FMenuEntryParams AddOnClickParams;
 				AddOnClickParams.LabelOverride = DisplayInfo.Label;
 				AddOnClickParams.ToolTipOverride = DisplayInfo.ToolTip;
+				AddOnClickParams.IconOverride = DisplayInfo.Icon;
 				AddOnClickParams.UserInterfaceActionType = EUserInterfaceActionType::Button; 
 				AddOnClickParams.DirectActions = {
 					FExecuteAction::CreateLambda([Option, Args = MoveTemp(Args)](){ Args.OnItemsSelected.Execute(Option->GetSelectableItems()); }),
@@ -248,6 +248,7 @@ namespace UE::ConcertSharedSlate
 				MenuBuilder.AddMenuEntry(AddOnClickParams);
 				break;
 			}
+			
 		default:
 			checkNoEntry();
 		}
@@ -261,17 +262,20 @@ namespace UE::ConcertSharedSlate
 		{
 			return Args.GetItemDisplayString.Execute(Left) < Args.GetItemDisplayString.Execute(Right); 
 		});
-
+		
 		// Add the objects to the MenuBuilder
+		const ESourceType SourceType = Source->GetDisplayInfo().SourceType;
+		const bool bUseButton = SourceType == ESourceType::ShowAsList;
+		
 		for (const TItemType& Item : SortedItems)
 		{
 			// Gray out entries that will not do anything.
-			// Remember: we're not creating a context menu but a combo button next to the search bar.
-			// If it was just gone, users might be confused why suddenly there is a missing entry!
-			const FCanExecuteAction CanExecuteAction = FCanExecuteAction::CreateLambda([Item, IsItemSelected = Args.IsItemSelected]()
-			{
-				return !IsItemSelected.IsBound() || !IsItemSelected.Execute(Item);
-			});
+			const FCanExecuteAction CanExecuteAction = bUseButton
+				? FCanExecuteAction::CreateLambda([Item, IsItemSelected = Args.IsItemSelected](){ return !IsItemSelected.IsBound() || !IsItemSelected.Execute(Item); })
+				: FCanExecuteAction::CreateLambda([](){ return true; });
+			const FIsActionChecked IsActionChecked = bUseButton
+				? FIsActionChecked()
+				: FIsActionChecked::CreateLambda([Item, IsItemSelected = Args.IsItemSelected](){ return IsItemSelected.IsBound() && IsItemSelected.Execute(Item); });
 			MenuBuilder.AddMenuEntry(
 				FText::FromString(Args.GetItemDisplayString.Execute(Item)),
 				TAttribute<FText>(TAttribute<FText>::CreateLambda([CanExecuteAction](){ return CanExecuteAction.Execute() ? LOCTEXT("ItemSelected", "Item is already selected") : FText::GetEmpty(); })),
@@ -281,9 +285,11 @@ namespace UE::ConcertSharedSlate
 					{
 						SelectItemDelegate.Execute({ Item });
 					}),
-					CanExecuteAction),
+					CanExecuteAction,
+					IsActionChecked
+					),
 				NAME_None,
-				EUserInterfaceActionType::Button
+				bUseButton ? EUserInterfaceActionType::Button : EUserInterfaceActionType::ToggleButton
 				);
 		}
 	}
