Allow users to select the properties they want displayed in the UI:
- In bottom view, the filter button has been replaced by edit button
- Edit button lets user's select the properties they want to work with
- Top view no longer shows components

diff --git a/Engine/Plugins/Developer/Concert/ConcertApp/MultiUserClient/Source/MultiUserClient/Private/Replication/Misc/GlobalAuthorityCache.cpp b/Engine/Plugins/Developer/Concert/ConcertApp/MultiUserClient/Source/MultiUserClient/Private/Replication/Misc/GlobalAuthorityCache.cpp
--- a/Engine/Plugins/Developer/Concert/ConcertApp/MultiUserClient/Source/MultiUserClient/Private/Replication/Misc/GlobalAuthorityCache.cpp
+++ b/Engine/Plugins/Developer/Concert/ConcertApp/MultiUserClient/Source/MultiUserClient/Private/Replication/Misc/GlobalAuthorityCache.cpp
@@ -137,6 +137,43 @@ namespace UE::MultiUserClient
 		return Conflict == EAuthorityConflict::Allowed;
 	}
 
+	void FGlobalAuthorityCache::ForEachClientReferencingProperty(const FSoftObjectPath& Object, const FConcertPropertyChain& Property, TFunctionRef<EBreakBehavior(const FGuid& ClientId)> Callback) const
+	{
+		const TSet<FGuid>* ClientsWithStreams = RegisteredObjectsToClients.Find(Object);
+		if (!ClientsWithStreams)
+		{
+			return;
+		}
+
+		for (const FGuid& ClientId : *ClientsWithStreams)
+		{
+			const FReplicationClient* Client = ClientManager.FindClient(ClientId);
+			if (!ensure(Client))
+			{
+				continue;
+			}
+
+			const FConcertReplicatedObjectInfo* ObjectInfo = Client->GetStreamSynchronizer().GetServerState().ReplicatedObjects.Find(Object);
+			const bool bContainsProperty = ensureMsgf(ObjectInfo, TEXT("RegisteredObjectsToClients lied. Investigate."))
+				&& ObjectInfo->PropertySelection.ReplicatedProperties.Contains(Property);
+			if (bContainsProperty && Callback(ClientId) == EBreakBehavior::Break)
+			{
+				break;
+			}
+		}
+	}
+
+	bool FGlobalAuthorityCache::IsPropertyReferencedByAnyClientStream(const FSoftObjectPath& Object, const FConcertPropertyChain& Property) const
+	{
+		bool bIsReferenced = false;
+		ForEachClientReferencingProperty(Object, Property, [&](const FGuid& ClientId)
+		{
+			bIsReferenced = true;
+			return EBreakBehavior::Continue;
+		});
+		return bIsReferenced;
+	}
+
 	TOptional<FGuid> FGlobalAuthorityCache::GetClientWithAuthorityOverProperty(const FSoftObjectPath& Object, const TArray<FName>& PropertyChain) const
 	{
 		TOptional<FGuid> Result;
diff --git a/Engine/Plugins/Developer/Concert/ConcertApp/MultiUserClient/Source/MultiUserClient/Private/Replication/Misc/GlobalAuthorityCache.h b/Engine/Plugins/Developer/Concert/ConcertApp/MultiUserClient/Source/MultiUserClient/Private/Replication/Misc/GlobalAuthorityCache.h
--- a/Engine/Plugins/Developer/Concert/ConcertApp/MultiUserClient/Source/MultiUserClient/Private/Replication/Misc/GlobalAuthorityCache.h
+++ b/Engine/Plugins/Developer/Concert/ConcertApp/MultiUserClient/Source/MultiUserClient/Private/Replication/Misc/GlobalAuthorityCache.h
@@ -54,6 +54,11 @@ namespace UE::MultiUserClient
 		/** @return Whether the given client add the given property to the object without causing any conflicts. */
 		bool CanClientAddProperty(const FSoftObjectPath& Object, const FGuid& ClientId, const FConcertPropertyChain& Chain) const;
 
+		/** Iterates through every client whose stream is referencing Object's Property. */
+		void ForEachClientReferencingProperty(const FSoftObjectPath& Object, const FConcertPropertyChain& Property, TFunctionRef<EBreakBehavior(const FGuid& ClientId)> Callback) const;
+		/** @return Whether any client is referencing Object in a stream. */
+		bool IsPropertyReferencedByAnyClientStream(const FSoftObjectPath& Object, const FConcertPropertyChain& Property) const;
+
 		/** Gets the client that has authority over the given property, if there is any. */
 		TOptional<FGuid> GetClientWithAuthorityOverProperty(const FSoftObjectPath& Object, const FConcertPropertyChain& Property) const { return GetClientWithAuthorityOverProperty(Object, Property.GetPathToProperty()); }
 		/** Gets the client that has authority over the given property, if there is any. */
@@ -83,9 +88,9 @@ namespace UE::MultiUserClient
 		/** Used to obtain the clients and their states */
 		FReplicationClientManager& ClientManager;
 		
-		/** Maps objects that are owned to the clients that own them */
+		/** Maps objects that are owned to the clients that own them (have authority) */
 		TMap<FSoftObjectPath, TSet<FGuid>> OwnedObjectsToClients;
-		/** Maps objects that are owned to the clients that own them */
+		/** Maps objects that are owned to the clients that have them in the stream */
 		TMap<FSoftObjectPath, TSet<FGuid>> RegisteredObjectsToClients;
 		
 		/** Called when the cache changes for a specific client. */
diff --git a/Engine/Plugins/Developer/Concert/ConcertApp/MultiUserClient/Source/MultiUserClient/Private/Replication/Misc/PropertySelection/UserPropertySelectionSource.cpp b/Engine/Plugins/Developer/Concert/ConcertApp/MultiUserClient/Source/MultiUserClient/Private/Replication/Misc/PropertySelection/UserPropertySelectionSource.cpp
new file mode 100644
--- /dev/null
+++ b/Engine/Plugins/Developer/Concert/ConcertApp/MultiUserClient/Source/MultiUserClient/Private/Replication/Misc/PropertySelection/UserPropertySelectionSource.cpp
@@ -0,0 +1,114 @@
+﻿// Copyright Epic Games, Inc. All Rights Reserved.
+
+#include "UserPropertySelectionSource.h"
+
+#include "Replication/Client/ReplicationClientManager.h"
+#include "Replication/Editor/Model/IEditableReplicationStreamModel.h"
+#include "Replication/Editor/Model/Property/IPropertySource.h"
+
+namespace UE::MultiUserClient
+{
+	namespace Private
+	{
+		class FUserPropertySource
+			: public ConcertSharedSlate::IPropertySource
+			, public FNoncopyable
+		{
+		public:
+			
+			FUserPropertySource(
+				FSoftObjectPath Object,
+				const ConcertSharedSlate::IEditableReplicationStreamModel& InUserSelection UE_LIFETIMEBOUND,
+				const FReplicationClientManager& InClientManager UE_LIFETIMEBOUND
+				)
+				: Object(MoveTemp(Object))
+				, UserSelection(InUserSelection)
+				, ClientManager(InClientManager)
+			{}
+
+			//~ Begin IPropertySource Interface
+			virtual void EnumerateProperties(TFunctionRef<EBreakBehavior(const ConcertSharedSlate::FPropertyInfo& Property)> Delegate) const override
+			{
+				// The UI should list the properties the user explicitly selected...
+				EBreakBehavior BreakBehavior = EBreakBehavior::Continue;
+				UserSelection.ForEachProperty(Object, [&Delegate, &BreakBehavior](const FConcertPropertyChain& Property)
+				{
+					BreakBehavior = Delegate(ConcertSharedSlate::FPropertyInfo(Property));
+					return BreakBehavior;
+				});
+				if (BreakBehavior == EBreakBehavior::Break)
+				{
+					return;
+				}
+
+				// ... and the properties that are in use
+				
+				// The idea behind this is to avoid making unnecessary FConcertPropertyChain copies (clutters heap memory)
+				using FHash = uint32;
+				using FPropertyPtrArray = TArray<const FConcertPropertyChain*, TInlineAllocator<4>>;
+				TMap<FHash, FPropertyPtrArray, TInlineSetAllocator<512>> VisitedProperties;
+				
+				ClientManager.ForEachClient([this, &Delegate, &VisitedProperties](const FReplicationClient& Client)
+				{
+					const FConcertReplicatedObjectInfo* ObjectInfo = Client.GetStreamSynchronizer()
+						.GetServerState()
+						.ReplicatedObjects.Find(Object);
+					if (!ObjectInfo)
+					{
+						return EBreakBehavior::Continue;
+					}
+
+					for (const FConcertPropertyChain& PropertyChain : ObjectInfo->PropertySelection.ReplicatedProperties)
+					{
+						const FHash Hash = GetTypeHash(PropertyChain);
+						FPropertyPtrArray& PropertyArray = VisitedProperties.FindOrAdd(Hash);
+						const bool bListedByOtherStream = PropertyArray.ContainsByPredicate([&PropertyChain](const FConcertPropertyChain* Visited){ return PropertyChain == *Visited; });
+						
+						// Properties should not be listed multiple times
+						const bool bAlreadyListed = bListedByOtherStream
+							|| UserSelection.ContainsProperties(Object, { PropertyChain });
+						if (bAlreadyListed)
+						{
+							continue;
+						}
+
+						if (Delegate(ConcertSharedSlate::FPropertyInfo(PropertyChain)) == EBreakBehavior::Break)
+						{
+							return EBreakBehavior::Break;
+						}
+						PropertyArray.Add(&PropertyChain);
+					}
+					return EBreakBehavior::Continue;
+				});
+			}
+			//~ End IPropertySource Interface
+
+		private:
+
+			/** The object for which the properties are being displayed. */
+			const FSoftObjectPath Object;
+
+			/** Used to get the properties the user has selected. */
+			const ConcertSharedSlate::IEditableReplicationStreamModel& UserSelection;
+			/** Used to get client stream content. */
+			const FReplicationClientManager& ClientManager;
+		};
+	}
+	
+	FUserPropertySelectionSource::FUserPropertySelectionSource(
+		const ConcertSharedSlate::IEditableReplicationStreamModel& InUserSelection,
+		const FReplicationClientManager& InClientManager
+		)
+		: UserSelection(InUserSelection)
+		, ClientManager(InClientManager)
+	{}
+
+	void FUserPropertySelectionSource::ProcessPropertySource(
+		const ConcertSharedSlate::FPropertySourceContext& Context,
+		TFunctionRef<void(const ConcertSharedSlate::IPropertySource& Model)> Processor
+		) const
+	{
+		const Private::FUserPropertySource UserPropertySource(Context.Object.GetUniqueID(), UserSelection, ClientManager);
+		Processor(UserPropertySource);
+	}
+}
diff --git a/Engine/Plugins/Developer/Concert/ConcertApp/MultiUserClient/Source/MultiUserClient/Private/Replication/Misc/PropertySelection/UserPropertySelectionSource.h b/Engine/Plugins/Developer/Concert/ConcertApp/MultiUserClient/Source/MultiUserClient/Private/Replication/Misc/PropertySelection/UserPropertySelectionSource.h
new file mode 100644
--- /dev/null
+++ b/Engine/Plugins/Developer/Concert/ConcertApp/MultiUserClient/Source/MultiUserClient/Private/Replication/Misc/PropertySelection/UserPropertySelectionSource.h
@@ -0,0 +1,46 @@
+﻿// Copyright Epic Games, Inc. All Rights Reserved.
+
+#pragma once
+
+#include "Replication/Editor/Model/Property/IPropertySourceProcessor.h"
+
+#include "HAL//Platform.h"
+#include "Templates/UnrealTemplate.h"
+
+namespace UE::ConcertSharedSlate { class IEditableReplicationStreamModel; }
+
+namespace UE::MultiUserClient
+{
+	class FReplicationClientManager;
+	
+	/**
+	 * Injected into UI causing it to only display the properties that
+	 * - have been selected by the user
+	 * - are referenced by any client streams
+	 */
+	class FUserPropertySelectionSource
+		: public ConcertSharedSlate::IPropertySourceProcessor
+		, public FNoncopyable
+	{
+	public:
+
+		FUserPropertySelectionSource(
+			const ConcertSharedSlate::IEditableReplicationStreamModel& InUserSelection UE_LIFETIMEBOUND,
+			const FReplicationClientManager& InClientManager UE_LIFETIMEBOUND
+			);
+
+		//~ Begin IPropertySourceProcessor Interface
+		virtual void ProcessPropertySource(
+			const ConcertSharedSlate::FPropertySourceContext& Context,
+			TFunctionRef<void(const ConcertSharedSlate::IPropertySource& Model)> Processor
+			) const override;
+		//~ End IPropertySourceProcessor Interface
+
+	private:
+
+		/** Used to get the properties the user has selected. */
+		const ConcertSharedSlate::IEditableReplicationStreamModel& UserSelection;
+		/** Used to get client stream content and subscribe to changes. */
+		const FReplicationClientManager& ClientManager;
+	};
+}
diff --git a/Engine/Plugins/Developer/Concert/ConcertApp/MultiUserClient/Source/MultiUserClient/Private/Replication/Misc/PropertySelection/UserPropertySelector.cpp b/Engine/Plugins/Developer/Concert/ConcertApp/MultiUserClient/Source/MultiUserClient/Private/Replication/Misc/PropertySelection/UserPropertySelector.cpp
new file mode 100644
--- /dev/null
+++ b/Engine/Plugins/Developer/Concert/ConcertApp/MultiUserClient/Source/MultiUserClient/Private/Replication/Misc/PropertySelection/UserPropertySelector.cpp
@@ -0,0 +1,126 @@
+﻿// Copyright Epic Games, Inc. All Rights Reserved.
+
+#include "UserPropertySelector.h"
+
+#include "UserPropertySelectionSource.h"
+#include "Replication/ClientReplicationWidgetFactories.h"
+#include "Replication/Client/ReplicationClientManager.h"
+#include "Replication/Editor/Model/IEditableReplicationStreamModel.h"
+
+#include "UObject/Package.h"
+#include "UObject/UObjectGlobals.h"
+
+namespace UE::MultiUserClient
+{
+	FUserPropertySelector::FUserPropertySelector(FReplicationClientManager& InClientManager)
+		: ClientManager(InClientManager)
+		, PropertySelection(NewObject<UMultiUserReplicationStream>(GetTransientPackage(), NAME_None, RF_Transient | RF_Transactional))
+		, SelectionEditModel(
+			// Transact ClientContentStorage
+			ConcertClientSharedSlate::CreateTransactionalStreamModel(
+				// Read & write the stream data in ClientContentStorage
+				ConcertSharedSlate::CreateBaseStreamModel(
+					PropertySelection->MakeReplicationMapGetterAttribute()
+					),
+				*PropertySelection
+				)
+			)
+		, PropertyProcessor(MakeShared<FUserPropertySelectionSource>(*SelectionEditModel, InClientManager))
+	{
+		RegisterClient(ClientManager.GetLocalClient());
+		ClientManager.OnPostRemoteClientAdded().AddRaw(this, &FUserPropertySelector::OnClientAdded);
+		FCoreUObjectDelegates::OnObjectTransacted.AddRaw(this, &FUserPropertySelector::OnObjectTransacted);
+	}
+
+	FUserPropertySelector::~FUserPropertySelector()
+	{
+		ClientManager.OnPostRemoteClientAdded().RemoveAll(this);
+		ClientManager.ForEachClient([this](FReplicationClient& Client)
+		{
+			Client.GetStreamSynchronizer().OnServerStateChanged().RemoveAll(this);
+			return EBreakBehavior::Continue;
+		});
+		FCoreUObjectDelegates::OnObjectTransacted.RemoveAll(this);
+	}
+
+	void FUserPropertySelector::AddSelectedProperties(UObject* Object, TConstArrayView<FConcertPropertyChain> Properties)
+	{
+		SelectionEditModel->AddObjects({ Object });
+		SelectionEditModel->AddProperties({ Object }, Properties);
+
+		OnPropertySelectionChangedDelegate.Broadcast();
+	}
+
+	void FUserPropertySelector::RemoveSelectedProperties(UObject* Object, TConstArrayView<FConcertPropertyChain> Properties)
+	{
+		SelectionEditModel->RemoveProperties({ Object }, Properties);
+		if (!SelectionEditModel->HasAnyPropertyAssigned(Object))
+		{
+			SelectionEditModel->RemoveObjects({ Object });
+		}
+
+		const TSharedRef<ConcertSharedSlate::IEditableReplicationStreamModel> ClientEditModel = ClientManager.GetLocalClient().GetClientEditModel();
+		ClientEditModel->RemoveProperties({ Object }, Properties);
+		
+		OnPropertySelectionChangedDelegate.Broadcast();
+	}
+
+	bool FUserPropertySelector::IsPropertySelected(const FSoftObjectPath& Object, const FConcertPropertyChain& Property) const
+	{
+		return ClientManager.GetAuthorityCache().IsPropertyReferencedByAnyClientStream(Object, Property)
+			|| SelectionEditModel->HasProperty(Object, Property);
+	}
+
+	TSharedRef<ConcertSharedSlate::IPropertySourceProcessor> FUserPropertySelector::GetPropertySourceProcessor() const
+	{
+		return PropertyProcessor;
+	}
+
+	void FUserPropertySelector::AddReferencedObjects(FReferenceCollector& Collector)
+	{
+		Collector.AddReferencedObject(PropertySelection);
+	}
+
+	void FUserPropertySelector::RegisterClient(FReplicationClient& Client)
+	{
+		IClientStreamSynchronizer& StreamSynchronizer = Client.GetStreamSynchronizer();
+		TrackProperties(StreamSynchronizer.GetServerState());
+		StreamSynchronizer.OnServerStateChanged().AddRaw(this, &FUserPropertySelector::OnServerStateChanged, Client.GetEndpointId());
+	}
+
+	void FUserPropertySelector::OnServerStateChanged(const FGuid ClientId)
+	{
+		const FReplicationClient* Client = ClientManager.FindClient(ClientId);
+		if (ensure(Client))
+		{
+			TrackProperties(Client->GetStreamSynchronizer().GetServerState());
+		}
+	}
+
+	void FUserPropertySelector::TrackProperties(const FConcertObjectReplicationMap& ReplicationMap)
+	{
+		for (const TPair<FSoftObjectPath, FConcertReplicatedObjectInfo>& Pair : ReplicationMap.ReplicatedObjects)
+		{
+			UObject* Object = Pair.Key.ResolveObject();
+			// The object may come from a remote client that is in a different world than the local application
+			if (!Object)
+			{
+				continue;
+			}
+
+			for (const FConcertPropertyChain& Property : Pair.Value.PropertySelection.ReplicatedProperties)
+			{
+				AddSelectedProperties(Object, { Property });
+			}
+		}
+	}
+
+	void FUserPropertySelector::OnObjectTransacted(UObject* Object, const FTransactionObjectEvent& TransactionObjectEvent)
+	{
+		if (Object == PropertySelection)
+		{
+			// Refreshes UI
+			OnPropertySelectionChangedDelegate.Broadcast();
+		}
+	}
+}
diff --git a/Engine/Plugins/Developer/Concert/ConcertApp/MultiUserClient/Source/MultiUserClient/Private/Replication/Misc/PropertySelection/UserPropertySelector.h b/Engine/Plugins/Developer/Concert/ConcertApp/MultiUserClient/Source/MultiUserClient/Private/Replication/Misc/PropertySelection/UserPropertySelector.h
new file mode 100644
--- /dev/null
+++ b/Engine/Plugins/Developer/Concert/ConcertApp/MultiUserClient/Source/MultiUserClient/Private/Replication/Misc/PropertySelection/UserPropertySelector.h
@@ -0,0 +1,91 @@
+﻿// Copyright Epic Games, Inc. All Rights Reserved.
+
+#pragma once
+
+#include "Assets/MultiUserReplicationStream.h"
+
+#include "HAL/Platform.h"
+#include "Replication/Client/RemoteReplicationClient.h"
+#include "Templates/UnrealTemplate.h"
+#include "UObject/GCObject.h"
+#include "UObject/ObjectPtr.h"
+
+class FReferenceCollector;
+class FTransactionObjectEvent;
+class UObject;
+
+namespace UE::ConcertSharedSlate
+{
+	class IEditableReplicationStreamModel;
+	class IPropertySourceProcessor;
+}
+
+namespace UE::MultiUserClient
+{
+	class FUserPropertySelectionSource;
+	class FReplicationClientManager;
+	
+	/**
+	 * Manages the properties the user is iterating on in the replication session.
+	 * The bottom-half property section in the replication UI uses this to keep track of which properties the user has selected for which properties.
+	 *
+	 * Whenever any client adds a property to its stream, we'll assume the user is iterating on that property.
+	 * For this reason, we automatically will track the property as user-selected.
+	 */
+	class FUserPropertySelector
+		: public FNoncopyable
+		, public FGCObject
+	{
+	public:
+		
+		FUserPropertySelector(FReplicationClientManager& InClientManager UE_LIFETIMEBOUND);
+		~FUserPropertySelector();
+
+		/** Add Properties from the user's selection for Object. */
+		void AddSelectedProperties(UObject* Object, TConstArrayView<FConcertPropertyChain> Properties);
+		/** Removes Properties from the user's selection for Object. */
+		void RemoveSelectedProperties(UObject* Object, TConstArrayView<FConcertPropertyChain> Properties);
+		/** @return Whether Property is selected for Object. */
+		bool IsPropertySelected(const FSoftObjectPath& Object, const FConcertPropertyChain& Property) const;
+
+		TSharedRef<ConcertSharedSlate::IPropertySourceProcessor> GetPropertySourceProcessor() const;
+
+		DECLARE_MULTICAST_DELEGATE(FOnPropertySelectionChanged)
+		/** @return Event that broadcasts when the user property selection changes. */
+		FOnPropertySelectionChanged& OnPropertySelectionChanged() { return OnPropertySelectionChangedDelegate; }
+		
+		//~ Begin FGCObject Interface
+		virtual void AddReferencedObjects(FReferenceCollector& Collector) override;
+		virtual FString GetReferencerName() const override { return TEXT("FUserPropertySelector"); }
+		//~ Begin FGCObject Interface
+
+	private:
+
+		/** Used to remove deselected properties from local client's stream. */
+		FReplicationClientManager& ClientManager;
+
+		/** This underlying object saves the properties that user has selected. It allows for transactions. */
+		TObjectPtr<UMultiUserReplicationStream> PropertySelection;
+		/** This logic modifies PropertySelection. */
+		const TSharedRef<ConcertSharedSlate::IEditableReplicationStreamModel> SelectionEditModel;
+
+		/** Getter for UI to determine which properties to display. */
+		const TSharedRef<FUserPropertySelectionSource> PropertyProcessor;
+
+		/** Broadcasts when the user property selection changes. */
+		FOnPropertySelectionChanged OnPropertySelectionChangedDelegate;
+		
+		void OnClientAdded(FRemoteReplicationClient& Client) { RegisterClient(Client); }
+		/** Ensures that whenever the client's server state changes, its properties are tracked as user selected. */
+		void RegisterClient(FReplicationClient& Client);
+
+		/** Tracks all properties of the client as user selected. */
+		void OnServerStateChanged(const FGuid ClientId);
+		/** Adds all properties in the replication map as user selected. */
+		void TrackProperties(const FConcertObjectReplicationMap& ReplicationMap);
+
+		/** If PropertySelection is transacted, broadcast OnPropertySelectionChangedDelegate. */
+		void OnObjectTransacted(UObject* Object, const FTransactionObjectEvent& TransactionObjectEvent);
+	};
+}
+
diff --git a/Engine/Plugins/Developer/Concert/ConcertApp/MultiUserClient/Source/MultiUserClient/Private/Replication/MultiUserReplicationManager.cpp b/Engine/Plugins/Developer/Concert/ConcertApp/MultiUserClient/Source/MultiUserClient/Private/Replication/MultiUserReplicationManager.cpp
--- a/Engine/Plugins/Developer/Concert/ConcertApp/MultiUserClient/Source/MultiUserClient/Private/Replication/MultiUserReplicationManager.cpp
+++ b/Engine/Plugins/Developer/Concert/ConcertApp/MultiUserClient/Source/MultiUserClient/Private/Replication/MultiUserReplicationManager.cpp
@@ -75,9 +75,9 @@ namespace UE::MultiUserClient
 
 	void FMultiUserReplicationManager::OnLeaveSession(IConcertClientSession&)
 	{
-		// This clears the UI. The clients' IEditableReplicationStreamModels should no longer be referenced by anyone.
+		// This destroys the UI and tells any other potential system to stop referencing anything in ConnectedState (such as shared ptrs)...
 		SetConnectionStateAndBroadcast(EMultiUserReplicationConnectionState::Disconnected);
-		// Keep in mind the IEditableReplicationStreamModels were referenced by the UI so call this after clearing the UI.
+		// ... so now it is safe to destroy ConnectedState.
 		ConnectedState.Reset();
 	}
 
@@ -200,6 +200,7 @@ namespace UE::MultiUserClient
 		, QueryService(*InClient)
 		, ClientManager(InClient, InClient->GetConcertClient()->GetCurrentSession().ToSharedRef(), InDiscoveryContainer, QueryService.GetStreamAndAuthorityQueryService())
 		, MuteManager(*InClient, QueryService.GetMuteStateQueryService(), ClientManager.GetAuthorityCache())
+		, PropertySelector(ClientManager)
 		, ChangeLevelHandler(ClientManager.GetLocalClient().GetClientEditModel().Get())
 		, PreventReplicatedPropertyTransaction(*InClient, ClientManager, MuteManager)
 		, UserNotifier(ClientManager, MuteManager)
diff --git a/Engine/Plugins/Developer/Concert/ConcertApp/MultiUserClient/Source/MultiUserClient/Private/Replication/MultiUserReplicationManager.h b/Engine/Plugins/Developer/Concert/ConcertApp/MultiUserClient/Source/MultiUserClient/Private/Replication/MultiUserReplicationManager.h
--- a/Engine/Plugins/Developer/Concert/ConcertApp/MultiUserClient/Source/MultiUserClient/Private/Replication/MultiUserReplicationManager.h
+++ b/Engine/Plugins/Developer/Concert/ConcertApp/MultiUserClient/Source/MultiUserClient/Private/Replication/MultiUserReplicationManager.h
@@ -13,6 +13,7 @@
 
 #include "Misc/Optional.h"
 #include "Misc/PreventReplicatedPropertyTransaction.h"
+#include "Misc/PropertySelection/UserPropertySelector.h"
 #include "Templates/SharedPointer.h"
 #include "Templates/UnrealTemplate.h"
 
@@ -67,6 +68,10 @@ namespace UE::MultiUserClient
 		FMuteStateManager* GetMuteManager() { return ConnectedState ? &ConnectedState->MuteManager : nullptr; }
 		const FMuteStateManager* GetMuteManager() const { return ConnectedState ? &ConnectedState->MuteManager : nullptr; }
 
+		/** @note You're not supposed to keep any reference to the PropertySelector since it can become invalid depending on connection state. */
+		FUserPropertySelector* GetUserPropertySelector() { return ConnectedState ? &ConnectedState->PropertySelector : nullptr; }
+		const FUserPropertySelector* GetUserPropertySelector() const { return ConnectedState ? &ConnectedState->PropertySelector : nullptr; }
+
 		/** Called when the connection to the replication system changes. */
 		DECLARE_MULTICAST_DELEGATE_OneParam(FOnReplicationConnectionStateChanged, EMultiUserReplicationConnectionState /*NewState*/);
 		FOnReplicationConnectionStateChanged& OnReplicationConnectionStateChanged() { return OnReplicationConnectionStateChangedDelegate; }
@@ -115,6 +120,11 @@ namespace UE::MultiUserClient
 			/** Interacts with the mute global server mute system. */
 			FMuteStateManager MuteManager;
 
+			/**
+			 * Manages the properties the user is iterating on in the replication session.
+			 * The bottom-half property section in the replication UI uses this to keep track of which properties the user has selected for which properties.
+			 */
+			FUserPropertySelector PropertySelector;
 			/** Clears local client's registered objects when leaving map. */
 			FChangeLevelHandler ChangeLevelHandler;
 			/** Prevents recording of transactions that change properties that are being replicated by a client. */
diff --git a/Engine/Plugins/Developer/Concert/ConcertApp/MultiUserClient/Source/MultiUserClient/Private/Widgets/ActiveSession/Replication/Client/Multi/Columns/MultiStreamColumns_AssignProperties.cpp b/Engine/Plugins/Developer/Concert/ConcertApp/MultiUserClient/Source/MultiUserClient/Private/Widgets/ActiveSession/Replication/Client/Multi/Columns/MultiStreamColumns_AssignProperties.cpp
--- a/Engine/Plugins/Developer/Concert/ConcertApp/MultiUserClient/Source/MultiUserClient/Private/Widgets/ActiveSession/Replication/Client/Multi/Columns/MultiStreamColumns_AssignProperties.cpp
+++ b/Engine/Plugins/Developer/Concert/ConcertApp/MultiUserClient/Source/MultiUserClient/Private/Widgets/ActiveSession/Replication/Client/Multi/Columns/MultiStreamColumns_AssignProperties.cpp
@@ -106,7 +106,7 @@ namespace UE::MultiUserClient::MultiStreamColumns
 			
 			virtual TSharedRef<SWidget> GenerateColumnWidget(const FBuildArgs& InArgs) override
 			{
-				const TArray<TSoftObjectPtr<>> DisplayedObjects = MultiStreamEditor.Get()->GetEditorBase().GetSelectedObjects();
+				const TArray<TSoftObjectPtr<>> DisplayedObjects = InArgs.RowItem.RowData.GetContextObjects();
 				return SNew(SAssignPropertyComboBox, MultiStreamEditor.Get().ToSharedRef(), ConcertClient, ClientManager)
 					.DisplayedProperty(InArgs.RowItem.RowData.GetProperty())
 					.EditedObjects(DisplayedObjects)
diff --git a/Engine/Plugins/Developer/Concert/ConcertApp/MultiUserClient/Source/MultiUserClient/Private/Widgets/ActiveSession/Replication/Client/Multi/Columns/MultiStreamColumns_MuteToggle.cpp b/Engine/Plugins/Developer/Concert/ConcertApp/MultiUserClient/Source/MultiUserClient/Private/Widgets/ActiveSession/Replication/Client/Multi/Columns/MultiStreamColumns_MuteToggle.cpp
--- a/Engine/Plugins/Developer/Concert/ConcertApp/MultiUserClient/Source/MultiUserClient/Private/Widgets/ActiveSession/Replication/Client/Multi/Columns/MultiStreamColumns_MuteToggle.cpp
+++ b/Engine/Plugins/Developer/Concert/ConcertApp/MultiUserClient/Source/MultiUserClient/Private/Widgets/ActiveSession/Replication/Client/Multi/Columns/MultiStreamColumns_MuteToggle.cpp
@@ -40,7 +40,7 @@ namespace UE::MultiUserClient::MultiStreamColumns
 						.HeightOverride(16.f)
 						.HAlign(HAlign_Center)
 						.VAlign(VAlign_Center)
-						.ToolTipText(LOCTEXT("Mute.ToolTip", "Whether an object is supposed to replicate or not.\nTip: When you pause an actor, you can selectively unpause components."))
+						.ToolTipText(LOCTEXT("Mute.ToolTip", "Whether an object is supposed to replicate or not."))
 						[
 							SNew(SImage)
 							.Image(FMultiUserReplicationStyle::Get()->GetBrush(TEXT("MultiUser.Icons.Play")))
diff --git a/Engine/Plugins/Developer/Concert/ConcertApp/MultiUserClient/Source/MultiUserClient/Private/Widgets/ActiveSession/Replication/Client/Multi/SMultiClientView.cpp b/Engine/Plugins/Developer/Concert/ConcertApp/MultiUserClient/Source/MultiUserClient/Private/Widgets/ActiveSession/Replication/Client/Multi/SMultiClientView.cpp
--- a/Engine/Plugins/Developer/Concert/ConcertApp/MultiUserClient/Source/MultiUserClient/Private/Widgets/ActiveSession/Replication/Client/Multi/SMultiClientView.cpp
+++ b/Engine/Plugins/Developer/Concert/ConcertApp/MultiUserClient/Source/MultiUserClient/Private/Widgets/ActiveSession/Replication/Client/Multi/SMultiClientView.cpp
@@ -9,16 +9,20 @@
 #include "Replication/ReplicationWidgetFactories.h"
 #include "Replication/Client/ReplicationClient.h"
 #include "Replication/Client/ReplicationClientManager.h"
+#include "Replication/Editor/Model/Object/IObjectNameModel.h"
 #include "Replication/Editor/Model/PropertySource/SelectPropertyFromUClassModel.h"
 #include "Replication/Editor/View/IMultiObjectPropertyAssignmentView.h"
 #include "Replication/Editor/View/IMultiReplicationStreamEditor.h"
 #include "Replication/Editor/View/IReplicationStreamEditor.h"
 #include "Replication/Editor/Model/ObjectSource/ActorSelectionSourceModel.h"
+#include "Widgets/ActiveSession/Replication/Client/Context/ContextMenuUtils.h"
 #include "Widgets/ActiveSession/Replication/Client/Multi/Columns/MultiStreamColumns.h"
+#include "Widgets/ActiveSession/Replication/Client/PropertySelection/SPropertySelectionComboButton.h"
 #include "Widgets/ActiveSession/Replication/Client/SReplicationStatus.h"
-#include "Widgets/ActiveSession/Replication/Client/Context/ContextMenuUtils.h"
 
+#include "Widgets/Layout/SBox.h"
 #include "Widgets/SBoxPanel.h"
+#include "Widgets/Text/STextBlock.h"
 
 #define LOCTEXT_NAMESPACE "SMultiClientView"
 
@@ -32,12 +36,14 @@ namespace UE::MultiUserClient
 		)
 	{
 		ClientManager = InMultiUserReplicationManager.GetClientManager();
+		UserSelectedProperties = InMultiUserReplicationManager.GetUserPropertySelector();
 		StreamModel = MakeShared<FMultiStreamModel>(InDisplayClientsModel, *ClientManager);
-
+		SelectionModel = &InDisplayClientsModel;
 		ConcertClient = MoveTemp(InConcertClient);
+		
 		ClientManager->OnRemoteClientsChanged().AddSP(this, &SMultiClientView::RebuildClientSubscriptions);
-		SelectionModel = &InDisplayClientsModel;
 		SelectionModel->OnSelectionChanged().AddSP(this, &SMultiClientView::RebuildClientSubscriptions);
+		UserSelectedProperties->OnPropertySelectionChanged().AddRaw(this, &SMultiClientView::RefreshUI);
 
 		TSharedPtr<SVerticalBox> Content;
 		ChildSlot
@@ -71,6 +77,7 @@ namespace UE::MultiUserClient
 	{
 		ClientManager->OnRemoteClientsChanged().RemoveAll(this);
 		CleanClientSubscriptions();
+		UserSelectedProperties->OnPropertySelectionChanged().RemoveAll(this);
 	}
 
 	TSharedRef<SWidget> SMultiClientView::CreateEditorContent(const TSharedRef<IConcertClient>& InConcertClient, FMultiUserReplicationManager& InMultiUserReplicationManager)
@@ -78,6 +85,7 @@ namespace UE::MultiUserClient
 		using namespace UE::ConcertSharedSlate;
 
 		FMuteStateManager& MuteManager = *InMultiUserReplicationManager.GetMuteManager();
+		FUserPropertySelector& PropertySelector = *InMultiUserReplicationManager.GetUserPropertySelector();
 		
 		ObjectHierarchy = ConcertClientSharedSlate::CreateObjectHierarchyForComponentHierarchy();
 		const TSharedRef<IObjectNameModel> NameModel = ConcertClientSharedSlate::CreateEditorObjectNameModel();
@@ -97,22 +105,30 @@ namespace UE::MultiUserClient
 			const TSharedRef<IEditableReplicationStreamModel>& LocalStream = ClientManager->GetLocalClient().GetClientEditModel();
 			return StreamModel->GetEditableStreams().Contains(LocalStream) ? LocalStream.ToSharedPtr() : nullptr;
 		});
-		const TSharedRef<ConcertClientSharedSlate::FSelectPropertyFromUClassModel> PropertySourceModel = MakeShared<ConcertClientSharedSlate::FSelectPropertyFromUClassModel>();
 		
-		ConcertClientSharedSlate::FFilterablePropertyTreeViewParams TreeViewParams
+		FCreatePropertyTreeViewParams TreeViewParams
 		{
-			.AdditionalPropertyColumns =
+			.PropertyColumns =
 			{
 				ReplicationColumns::Property::LabelColumn(),
 				MultiStreamColumns::AssignPropertyColumn(MultiStreamEditorAttribute, InConcertClient, *ClientManager)
 			},
-			.CreateCategoryRow = CreateDefaultCategoryGenerator(NameModel)
+			.CreateCategoryRow = CreateDefaultCategoryGenerator(NameModel),
 		};
-		TSharedRef<IPropertyTreeView> PropertyTreeView = CreateFilterablePropertyTreeView(MoveTemp(TreeViewParams));
+		TreeViewParams.LeftOfPropertySearchBar.Widget = SAssignNew(PropertySelectionButton, SPropertySelectionComboButton, PropertySelector)
+			.GetObjectDisplayString_Lambda([NameModel](const TSoftObjectPtr<>& Object){ return NameModel->GetObjectDisplayName(Object); });
+		TreeViewParams.NoItemsContent.Widget = CreateNoPropertiesWarning();
+		const TSharedRef<IPropertyTreeView> PropertyTreeView = CreateSearchablePropertyTreeView(MoveTemp(TreeViewParams));
 		
-		TSharedRef<IMultiObjectPropertyAssignmentView> PropertyAssignmentView = CreateMultiObjectAssignmentView(
+		const TSharedRef<IPropertySourceProcessor> PropertySourceModel = PropertySelector.GetPropertySourceProcessor();
+		PropertyAssignmentView = CreateMultiObjectAssignmentView(
 			{ .PropertyTreeView = PropertyTreeView, .ObjectHierarchy = ObjectHierarchy, .PropertySource = PropertySourceModel}
 			);
+		PropertyAssignmentView->OnObjectGroupsChanged().AddLambda([this]()
+		{
+			PropertySelectionButton->RefreshSelectableProperties(PropertyAssignmentView->GetDisplayedGroups());
+		});
+		
 		
 		FCreateMultiStreamEditorParams Params
 		{
@@ -122,11 +138,10 @@ namespace UE::MultiUserClient
 			.PropertySource = PropertySourceModel,
 			.GetAutoAssignToStreamDelegate = MoveTemp(GetAutoAssignTargetDelegate)
 		};
-		
 		FCreateViewerParams ViewerParams
 		{
-			.PropertyAssignmentView = StaticCastSharedRef<IPropertyAssignmentView>(PropertyAssignmentView),
-			.ObjectHierarchy = ObjectHierarchy, // This makes actors have children in the top view
+			.PropertyAssignmentView = PropertyAssignmentView.ToSharedRef(),
+			// .ObjectHierarchy Do not assign so we only show the actors
 			.NameModel = NameModel, // This makes actors use their labels, and components use the names given in the BP editor
 			.OnExtendObjectsContextMenu = FExtendObjectMenu::CreateSP(this, &SMultiClientView::ExtendObjectContextMenu),
 			.ObjectColumns =
@@ -136,12 +151,23 @@ namespace UE::MultiUserClient
 			},
 			.ShouldDisplayObjectDelegate = FShouldDisplayObject::CreateSP(this, &SMultiClientView::ShouldDisplayObject)
 		};
-		
 		StreamEditor = CreateBaseMultiStreamEditor(MoveTemp(Params), MoveTemp(ViewerParams));
 		check(StreamEditor);
+		
 		return StreamEditor.ToSharedRef();
 	}
 
+	TSharedRef<SWidget> SMultiClientView::CreateNoPropertiesWarning() const
+	{
+		return SNew(SBox)
+			.HAlign(HAlign_Center)
+			.VAlign(VAlign_Center)
+			[
+				SNew(STextBlock)
+				.Text(LOCTEXT("NoProperties", "Use Edit button to add replicated properties"))
+			];
+	}
+
 	TSet<FGuid> SMultiClientView::GetDisplayClientIds() const
 	{
 		TSet<FGuid> ClientIds;
@@ -174,8 +200,8 @@ namespace UE::MultiUserClient
 		{
 			if (SelectionModel->ContainsClient(Client.GetEndpointId()))
 			{
-				Client.OnModelChanged().AddSP(this, &SMultiClientView::OnClientChanged);
-				Client.OnHierarchyNeedsRefresh().AddRaw(this, &SMultiClientView::OnHierarchyNeedsRefresh);
+				Client.OnModelChanged().AddSP(this, &SMultiClientView::RefreshUI);
+				Client.OnHierarchyNeedsRefresh().AddRaw(this, &SMultiClientView::RefreshUI);
 			}
 			
 			return EBreakBehavior::Continue;
@@ -192,16 +218,8 @@ namespace UE::MultiUserClient
 		});
 	}
 
-	void SMultiClientView::OnClientChanged() const
-	{
-		// When reassignment operations complete, the content of the columns changes so a resort is required.
-		StreamEditor->GetEditorBase().RequestObjectColumnResort(MultiStreamColumns::AssignedClientsColumnId);
-		StreamEditor->GetEditorBase().RequestPropertyColumnResort(MultiStreamColumns::AssignPropertyColumnId);
-	}
-
-	void SMultiClientView::OnHierarchyNeedsRefresh() const
+	void SMultiClientView::RefreshUI() const
 	{
-		// It's a bit excessive to refresh all objects when the hierarchy might have changed but it's simple (and only happens once at end of tick)
 		StreamEditor->GetEditorBase().Refresh();
 	}
 
diff --git a/Engine/Plugins/Developer/Concert/ConcertApp/MultiUserClient/Source/MultiUserClient/Private/Widgets/ActiveSession/Replication/Client/Multi/SMultiClientView.h b/Engine/Plugins/Developer/Concert/ConcertApp/MultiUserClient/Source/MultiUserClient/Private/Widgets/ActiveSession/Replication/Client/Multi/SMultiClientView.h
--- a/Engine/Plugins/Developer/Concert/ConcertApp/MultiUserClient/Source/MultiUserClient/Private/Widgets/ActiveSession/Replication/Client/Multi/SMultiClientView.h
+++ b/Engine/Plugins/Developer/Concert/ConcertApp/MultiUserClient/Source/MultiUserClient/Private/Widgets/ActiveSession/Replication/Client/Multi/SMultiClientView.h
@@ -5,12 +5,19 @@
 #include "Replication/Editor/UnrealEditor/HideObjectsNotInWorldLogic.h"
 #include "Widgets/DeclarativeSyntaxSupport.h"
 #include "Widgets/SCompoundWidget.h"
+#include "Widgets/ActiveSession/Replication/Client/PropertySelection/RootPropertySourceModel.h"
+
+namespace UE::MultiUserClient
+{
+	class FUserPropertySelector;
+}
 
 class IConcertClient;
 class FMenuBuilder;
 
 namespace UE::ConcertSharedSlate
 {
+	class IMultiObjectPropertyAssignmentView;
 	class IObjectHierarchyModel;
 	class IMultiReplicationStreamEditor;
 	class IEditableReplicationStreamModel;
@@ -24,6 +31,7 @@ namespace UE::MultiUserClient
 	class FReplicationClient;
 	class FReplicationClientManager;
 	class IClientSelectionModel;
+	class SPropertySelectionComboButton;
 
 	/** Displays a selection of clients. */
 	class SMultiClientView
@@ -44,6 +52,7 @@ namespace UE::MultiUserClient
 
 		TSharedPtr<IConcertClient> ConcertClient;
 		FReplicationClientManager* ClientManager = nullptr;
+		FUserPropertySelector* UserSelectedProperties = nullptr;
 		IClientSelectionModel* SelectionModel = nullptr;
 		
 		/** Combines the clients */
@@ -53,11 +62,20 @@ namespace UE::MultiUserClient
 		/** Used by widgets in columns. */
 		TSharedPtr<ConcertSharedSlate::IObjectHierarchyModel> ObjectHierarchy;
 
+		/**
+		 * This combo button is shown to the left of the search bar in the bottom half of the replication UI.
+		 * It allows users to specify the properties they want to work on (i.e. these properties should be shown in the property view).
+		 */
+		TSharedPtr<SPropertySelectionComboButton> PropertySelectionButton;
+		/** Displays the properties for the objects displayed in the top view. */
+		TSharedPtr<ConcertSharedSlate::IMultiObjectPropertyAssignmentView> PropertyAssignmentView;
+
 		/** This logic helps us decide whether an object should be displayed and lets us know that the object list needs to be refreshed (e.g. due to world change). */
 		ConcertClientSharedSlate::FHideObjectsNotInWorldLogic HideObjectsNotInEditorWorld;
 
 		/** Creates this widget's editor content */
 		TSharedRef<SWidget> CreateEditorContent(const TSharedRef<IConcertClient>& InConcertClient, FMultiUserReplicationManager& InMultiUserReplicationManager);
+		TSharedRef<SWidget> CreateNoPropertiesWarning() const;
 
 		// SClientToolbar attributes
 		TSet<FGuid> GetDisplayClientIds() const;
@@ -65,8 +83,7 @@ namespace UE::MultiUserClient
 		
 		void RebuildClientSubscriptions();
 		void CleanClientSubscriptions() const;
-		void OnClientChanged() const;
-		void OnHierarchyNeedsRefresh() const;
+		void RefreshUI() const;
 		
 		/** Adds additional entries to the context menu for the object tree view. */
 		void ExtendObjectContextMenu(FMenuBuilder& MenuBuilder, TConstArrayView<TSoftObjectPtr<>> ContextObjects) const;
diff --git a/Engine/Plugins/Developer/Concert/ConcertApp/MultiUserClient/Source/MultiUserClient/Private/Widgets/ActiveSession/Replication/Client/PropertySelection/AllRootPropertiesSource.cpp b/Engine/Plugins/Developer/Concert/ConcertApp/MultiUserClient/Source/MultiUserClient/Private/Widgets/ActiveSession/Replication/Client/PropertySelection/AllRootPropertiesSource.cpp
new file mode 100644
--- /dev/null
+++ b/Engine/Plugins/Developer/Concert/ConcertApp/MultiUserClient/Source/MultiUserClient/Private/Widgets/ActiveSession/Replication/Client/PropertySelection/AllRootPropertiesSource.cpp
@@ -0,0 +1,58 @@
+﻿// Copyright Epic Games, Inc. All Rights Reserved.
+
+#include "AllRootPropertiesSource.h"
+
+namespace UE::MultiUserClient
+{
+	namespace Private
+	{
+		static TArray<FUserSelectableProperty> PrebuildOptions(const ConcertSharedSlate::FObjectGroup& ObjectGroup, UClass& SharedClass)
+		{
+			TMap<FConcertPropertyChain, TArray<FConcertPropertyChain>> RootToChildren;
+			
+			ConcertClientSharedSlate::FReplicatablePropertySource(&SharedClass)
+				.EnumerateProperties([&RootToChildren](const ConcertSharedSlate::FPropertyInfo& Info)
+				{
+					const FConcertPropertyChain& Property = Info.Property;
+					if (Info.Property.IsRootProperty())
+					{
+						RootToChildren.Add(Info.Property);
+					}
+					else
+					{
+						RootToChildren[Info.Property.GetRootParent()].Add(Info.Property);
+					}
+					return EBreakBehavior::Continue;
+				});
+			
+			TArray<FUserSelectableProperty> Result;
+			for (TPair<FConcertPropertyChain, TArray<FConcertPropertyChain>>& Pair : RootToChildren)
+			{
+				Pair.Value.Add(Pair.Key);
+				Result.Emplace(ObjectGroup, MoveTemp(Pair.Key), MoveTemp(Pair.Value));
+			}
+			return Result;
+		}
+	}
+	
+	FAllRootPropertiesSource::FAllRootPropertiesSource(ConcertSharedSlate::FBaseDisplayInfo BaseDisplayInfo, const ConcertSharedSlate::FObjectGroup& ObjectGroup, UClass& SharedClass)
+		: BaseDisplayInfo(BaseDisplayInfo)
+		, Options(Private::PrebuildOptions(ObjectGroup, SharedClass))
+	{}
+
+	ConcertSharedSlate::FSourceDisplayInfo FAllRootPropertiesSource::GetDisplayInfo() const
+	{
+		return { { BaseDisplayInfo }, ConcertSharedSlate::ESourceType::ShowAsToggleButtonList };
+	}
+
+	void FAllRootPropertiesSource::EnumerateSelectableItems(TFunctionRef<EBreakBehavior(const FUserSelectableProperty& SelectableOption)> Delegate) const
+	{
+		for (const FUserSelectableProperty& Option : Options)
+		{
+			if (Delegate(Option) == EBreakBehavior::Break)
+			{
+				break;
+			}
+		}
+	}
+}
diff --git a/Engine/Plugins/Developer/Concert/ConcertApp/MultiUserClient/Source/MultiUserClient/Private/Widgets/ActiveSession/Replication/Client/PropertySelection/AllRootPropertiesSource.h b/Engine/Plugins/Developer/Concert/ConcertApp/MultiUserClient/Source/MultiUserClient/Private/Widgets/ActiveSession/Replication/Client/PropertySelection/AllRootPropertiesSource.h
new file mode 100644
--- /dev/null
+++ b/Engine/Plugins/Developer/Concert/ConcertApp/MultiUserClient/Source/MultiUserClient/Private/Widgets/ActiveSession/Replication/Client/PropertySelection/AllRootPropertiesSource.h
@@ -0,0 +1,42 @@
+﻿// Copyright Epic Games, Inc. All Rights Reserved.
+
+#pragma once
+
+#include "RootPropertySourceModel.h"
+#include "IPropertyItemSource.h"
+#include "Model/Item/IItemSourceModel.h"
+#include "Replication/Editor/Model/PropertySource/ReplicatablePropertySource.h"
+
+#include "Templates/UnrealTemplate.h"
+
+class UClass;
+
+namespace UE::MultiUserClient
+{
+	/** Adapts FReplicatablePropertySource. */
+	class FAllRootPropertiesSource
+		: public IPropertyItemSource
+		, public FNoncopyable
+	{
+	public:
+
+		FAllRootPropertiesSource(
+			ConcertSharedSlate::FBaseDisplayInfo BaseDisplayInfo, const ConcertSharedSlate::FObjectGroup& ObjectGroup,
+			UClass& SharedClass
+			);
+
+		//~ Begin IPropertyItemSource Interface
+		virtual ConcertSharedSlate::FSourceDisplayInfo GetDisplayInfo() const override;
+		virtual void EnumerateSelectableItems(TFunctionRef<EBreakBehavior(const FUserSelectableProperty& SelectableOption)> Delegate) const override;
+		//~ End IPropertyItemSource Interface
+
+	private:
+
+		/** Passed in display info about ObjectGroup. */
+		const ConcertSharedSlate::FBaseDisplayInfo BaseDisplayInfo;
+
+		/** The options being iterated. */
+		const TArray<FUserSelectableProperty> Options;
+	};
+}
+
diff --git a/Engine/Plugins/Developer/Concert/ConcertApp/MultiUserClient/Source/MultiUserClient/Private/Widgets/ActiveSession/Replication/Client/PropertySelection/IPropertyItemSource.h b/Engine/Plugins/Developer/Concert/ConcertApp/MultiUserClient/Source/MultiUserClient/Private/Widgets/ActiveSession/Replication/Client/PropertySelection/IPropertyItemSource.h
new file mode 100644
--- /dev/null
+++ b/Engine/Plugins/Developer/Concert/ConcertApp/MultiUserClient/Source/MultiUserClient/Private/Widgets/ActiveSession/Replication/Client/PropertySelection/IPropertyItemSource.h
@@ -0,0 +1,31 @@
+﻿// Copyright Epic Games, Inc. All Rights Reserved.
+
+#pragma once
+
+#include "Model/Item/IItemSourceModel.h"
+#include "Replication/Data/ConcertPropertySelection.h"
+#include "Replication/Editor/View/IPropertyAssignmentView.h"
+
+#include "Containers/Array.h"
+
+namespace UE::MultiUserClient
+{
+	struct FUserSelectableProperty
+	{
+		/** The objects for which this property is selected */
+		ConcertSharedSlate::FObjectGroup ObjectGroup;
+		/** The property shown to the user as option. */
+		FConcertPropertyChain RootProperty;
+		/** The properties that become (un)selected when the user (un)selects RootProperty. Contains RootProperty. */
+		TArray<FConcertPropertyChain> PropertiesToAdd;
+
+		FUserSelectableProperty(ConcertSharedSlate::FObjectGroup ObjectGroup, FConcertPropertyChain Property, TArray<FConcertPropertyChain> PropertiesToAdd)
+			: ObjectGroup(MoveTemp(ObjectGroup))
+			, RootProperty(MoveTemp(Property))
+			, PropertiesToAdd(MoveTemp(PropertiesToAdd))
+		{}
+	};
+	
+	using IPropertyItemSource = ConcertSharedSlate::IItemSourceModel<FUserSelectableProperty>;
+}
+
diff --git a/Engine/Plugins/Developer/Concert/ConcertApp/MultiUserClient/Source/MultiUserClient/Private/Widgets/ActiveSession/Replication/Client/PropertySelection/RootPropertySourceModel.cpp b/Engine/Plugins/Developer/Concert/ConcertApp/MultiUserClient/Source/MultiUserClient/Private/Widgets/ActiveSession/Replication/Client/PropertySelection/RootPropertySourceModel.cpp
new file mode 100644
--- /dev/null
+++ b/Engine/Plugins/Developer/Concert/ConcertApp/MultiUserClient/Source/MultiUserClient/Private/Widgets/ActiveSession/Replication/Client/PropertySelection/RootPropertySourceModel.cpp
@@ -0,0 +1,47 @@
+﻿// Copyright Epic Games, Inc. All Rights Reserved.
+
+#include "RootPropertySourceModel.h"
+
+#include "AllRootPropertiesSource.h"
+#include "Replication/Editor/View/IPropertyAssignmentView.h"
+
+#include "Styling/SlateIconFinder.h"
+
+#define LOCTEXT_NAMESPACE "FRootPropertySourceModel"
+
+namespace UE::MultiUserClient
+{
+	FRootPropertySourceModel::FRootPropertySourceModel(FGetObjectDisplayString InGetObjectDisplayStringDelegate)
+		: GetObjectDisplayStringDelegate(MoveTemp(InGetObjectDisplayStringDelegate))
+	{
+		check(GetObjectDisplayStringDelegate.IsBound());
+	}
+
+	void FRootPropertySourceModel::RefreshSelectableProperties(TConstArrayView<ConcertSharedSlate::FObjectGroup> DisplayedObjectGroups)
+	{
+		PerObjectGroup_CommonPropertiesSources.Empty(DisplayedObjectGroups.Num());
+		PerObjectGroup_AllPropertiesSources.Empty(DisplayedObjectGroups.Num());
+		
+		for (const ConcertSharedSlate::FObjectGroup& Group : DisplayedObjectGroups)
+		{
+			const UObject* Object = Group.Group.IsEmpty() ? nullptr : Group.Group[0].Get();
+			if (!ensure(Object))
+			{
+				continue;
+			}
+
+			UClass* Class = Object->GetClass();
+			const FText Label = GetObjectDisplayStringDelegate.Execute(Group.Group[0]);
+			const ConcertSharedSlate::FBaseDisplayInfo DisplayInfo
+			{
+				Label,
+				FText::Format(LOCTEXT("ToolTipFmt", "Edit properties for {0}"), Label),
+				FSlateIconFinder::FindIconForClass(Class)
+			};
+			
+			PerObjectGroup_AllPropertiesSources.Add(MakeShared<FAllRootPropertiesSource>(DisplayInfo, Group, *Class));
+		}
+	}
+}
+
+#undef LOCTEXT_NAMESPACE
\ No newline at end of file
diff --git a/Engine/Plugins/Developer/Concert/ConcertApp/MultiUserClient/Source/MultiUserClient/Private/Widgets/ActiveSession/Replication/Client/PropertySelection/RootPropertySourceModel.h b/Engine/Plugins/Developer/Concert/ConcertApp/MultiUserClient/Source/MultiUserClient/Private/Widgets/ActiveSession/Replication/Client/PropertySelection/RootPropertySourceModel.h
new file mode 100644
--- /dev/null
+++ b/Engine/Plugins/Developer/Concert/ConcertApp/MultiUserClient/Source/MultiUserClient/Private/Widgets/ActiveSession/Replication/Client/PropertySelection/RootPropertySourceModel.h
@@ -0,0 +1,46 @@
+﻿// Copyright Epic Games, Inc. All Rights Reserved.
+
+#pragma once
+
+#include "IPropertyItemSource.h"
+#include "Model/Item/SourceSelectionCategory.h"
+#include "Replication/Editor/View/IPropertyAssignmentView.h"
+#include "SelectionDelegates.h"
+
+#include "Templates/UnrealTemplate.h"
+
+namespace UE::ConcertSharedSlate { struct FObjectGroup; }
+
+namespace UE::MultiUserClient
+{
+	/**
+	 * This model is used to build the combo button to the left of the search bar in the bottom half of the replication UI.
+	 * It allows users to specify the properties they want to work on (i.e. these properties should be shown in the property view).
+	 */
+	class FRootPropertySourceModel : public FNoncopyable
+	{
+	public:
+
+		FRootPropertySourceModel(FGetObjectDisplayString InGetObjectDisplayStringDelegate);
+		
+		/** Refreshes the properties that the user can select given the objects currently displayed in the bottom view. */
+		void RefreshSelectableProperties(TConstArrayView<ConcertSharedSlate::FObjectGroup> DisplayedObjectGroups);
+
+		/** Contains one source for each object group. Each source displays common properties, configured in the project settings. */
+		const TArray<TSharedRef<IPropertyItemSource>>& GetPerObjectGroup_CommonPropertiesSources() const { return PerObjectGroup_CommonPropertiesSources; }
+		/** Contains one source for each object group. Each source displays all root properties in that class. */
+		const TArray<TSharedRef<IPropertyItemSource>>& GetPerObjectGroup_AllPropertiesSources() const { return PerObjectGroup_AllPropertiesSources; }
+
+	private:
+		
+		/** Determines the display string of an object. */
+		const FGetObjectDisplayString GetObjectDisplayStringDelegate;
+
+		// TODO UE-216094: Populate this array.
+		/** Contains one source for each object group. Each source displays common properties, configured in the project settings. */
+		TArray<TSharedRef<IPropertyItemSource>> PerObjectGroup_CommonPropertiesSources;
+		/** Contains one source for each object group. Each source displays all root properties in that class. */
+		TArray<TSharedRef<IPropertyItemSource>> PerObjectGroup_AllPropertiesSources;
+	};
+}
+
diff --git a/Engine/Plugins/Developer/Concert/ConcertApp/MultiUserClient/Source/MultiUserClient/Private/Widgets/ActiveSession/Replication/Client/PropertySelection/SPropertySelectionComboButton.cpp b/Engine/Plugins/Developer/Concert/ConcertApp/MultiUserClient/Source/MultiUserClient/Private/Widgets/ActiveSession/Replication/Client/PropertySelection/SPropertySelectionComboButton.cpp
new file mode 100644
--- /dev/null
+++ b/Engine/Plugins/Developer/Concert/ConcertApp/MultiUserClient/Source/MultiUserClient/Private/Widgets/ActiveSession/Replication/Client/PropertySelection/SPropertySelectionComboButton.cpp
@@ -0,0 +1,153 @@
+﻿// Copyright Epic Games, Inc. All Rights Reserved.
+
+#include "SPropertySelectionComboButton.h"
+
+#include "RootPropertySourceModel.h"
+#include "Model/Item/SourceModelBuilders.h"
+#include "Replication/Data/ConcertPropertySelection.h"
+#include "Replication/Misc/PropertySelection/UserPropertySelector.h"
+
+#include "Algo/AllOf.h"
+#include "ScopedTransaction.h"
+
+#define LOCTEXT_NAMESPACE "SPropertySelectionComboButton"
+
+namespace UE::MultiUserClient
+{
+	namespace Private
+	{
+		static bool IsSelected(const FUserPropertySelector& UserSelection, const FUserSelectableProperty& Property)
+		{
+			return Algo::AllOf(Property.ObjectGroup.Group, [&UserSelection, &Property](const TSoftObjectPtr<>& Object)
+			{
+				return UserSelection.IsPropertySelected(Object.GetUniqueID(), Property.RootProperty);
+			});
+		}
+
+		static void ToggleSelection(FUserPropertySelector& UserSelection, const FUserSelectableProperty& Property)
+		{
+			const bool bIsSelected = IsSelected(UserSelection, Property);
+			for (const TSoftObjectPtr<>& SoftObjectPtr : Property.ObjectGroup.Group)
+			{
+				UObject* Object = SoftObjectPtr.Get();
+				if (bIsSelected)
+				{
+					UserSelection.RemoveSelectedProperties(Object, Property.PropertiesToAdd);
+				}
+				else
+				{
+					UserSelection.AddSelectedProperties(Object, Property.PropertiesToAdd);
+				}
+			}
+			
+		}
+	}
+	
+	void SPropertySelectionComboButton::Construct(const FArguments& InArgs, FUserPropertySelector& InPropertySelector)
+	{
+		PropertySelector = &InPropertySelector;
+		PropertySourceModel = MakeUnique<FRootPropertySourceModel>(InArgs._GetObjectDisplayString);
+		
+		ChildSlot
+		[
+			SNew(SPositiveActionButton)
+			.Text(LOCTEXT("Edit.Label", "Edit")) 
+			.ToolTipText(LOCTEXT("Edit.ToolTip", "Select properties you want to work with"))
+			.OnGetMenuContent(this, &SPropertySelectionComboButton::MakeMenu)
+		];
+	}
+
+	void SPropertySelectionComboButton::RefreshSelectableProperties(TConstArrayView<ConcertSharedSlate::FObjectGroup> DisplayedObjectGroups) const
+	{
+		PropertySourceModel->RefreshSelectableProperties(DisplayedObjectGroups);
+	}
+
+	TSharedRef<SWidget> SPropertySelectionComboButton::MakeMenu() const
+	{
+		FMenuBuilder MenuBuilder(false, nullptr);
+		const FModelBuilder::FItemPickerArgs PickerArgs = MakePickerArguments();
+		
+		MenuBuilder.AddSeparator();
+		MenuBuilder.AddWidget(
+			SNew(SBox)
+			.Padding(3.f)
+			[
+				SNew(STextBlock)
+				.Text(LOCTEXT("Instructions", "Add a property you want to replicate here."))
+				.ColorAndOpacity(FSlateColor::UseSubduedForeground())
+			],
+			FText::GetEmpty()
+		);
+		
+		MenuBuilder.BeginSection(NAME_None, LOCTEXT("Section.CommonProperties", "Common Properties"));
+		MenuBuilder.AddMenuEntry(
+			LOCTEXT("Configure.Label", "Configure Common Properties"),
+			LOCTEXT("Configure.ToolTip", "Opens the project settings in which you can set up common properties"),
+			FSlateIcon(),
+			FUIAction(
+				FExecuteAction::CreateLambda([]()
+				{
+					// TODO UE-216094: Open project settings
+				}))
+			);
+		for (const TSharedRef<IPropertyItemSource>& ItemSource : PropertySourceModel->GetPerObjectGroup_CommonPropertiesSources())
+		{
+			FModelBuilder::AddOptionToMenu(ItemSource, PickerArgs, MenuBuilder);
+		}
+		MenuBuilder.EndSection();
+		
+		MenuBuilder.BeginSection(NAME_None, LOCTEXT("Section.AllProperties", "All Properties"));
+		for (const TSharedRef<IPropertyItemSource>& ItemSource : PropertySourceModel->GetPerObjectGroup_AllPropertiesSources())
+		{
+			FModelBuilder::AddOptionToMenu(ItemSource, PickerArgs, MenuBuilder);
+		}
+		MenuBuilder.EndSection();
+		
+		return MenuBuilder.MakeWidget();
+	}
+
+	ConcertSharedSlate::FSourceModelBuilders<FUserSelectableProperty>::FItemPickerArgs SPropertySelectionComboButton::MakePickerArguments() const
+	{
+		return FModelBuilder::FItemPickerArgs
+		(
+			FModelBuilder::FOnItemsSelected::CreateSP(this, &SPropertySelectionComboButton::OnItemsSelected),
+			FModelBuilder::FGetItemDisplayString::CreateLambda([](const FUserSelectableProperty& Item) -> FString
+			{
+				UObject* Object = Item.ObjectGroup.Group.IsEmpty() ? nullptr : Item.ObjectGroup.Group[0].Get();
+				if (!Object)
+				{
+					return TEXT("Unknown object");
+				}
+
+				const FProperty* Property = Item.RootProperty.ResolveProperty(*Object->GetClass());
+				return Property
+					? Property->GetDisplayNameText().ToString()
+					: TEXT("Unknown property");
+			}),
+			FModelBuilder::FGetItemIcon{},
+			FModelBuilder::FIsItemSelected::CreateLambda([this](const FUserSelectableProperty& Item)
+			{
+				return Private::IsSelected(*PropertySelector, Item);
+			})
+		);
+	}
+
+	void SPropertySelectionComboButton::OnItemsSelected(TArray<FUserSelectableProperty> Properties) const
+	{
+		// Usually there will only be one changed property
+		const FText Message = Properties.Num() == 1
+			? FText::Format(
+				LOCTEXT("Transaction.SingleFmt", "Change property '{0}'"),
+				FText::FromString(Properties[0].RootProperty.ToString(FConcertPropertyChain::EToStringMethod::LeafProperty))
+				)
+			: LOCTEXT("Transaction.Multi", "Change properties");
+		const FScopedTransaction Transaction(Message);
+		
+		for (const FUserSelectableProperty& Property : Properties)
+		{
+			Private::ToggleSelection(*PropertySelector, Property);
+		}
+	}
+}
+
+#undef LOCTEXT_NAMESPACE
\ No newline at end of file
diff --git a/Engine/Plugins/Developer/Concert/ConcertApp/MultiUserClient/Source/MultiUserClient/Private/Widgets/ActiveSession/Replication/Client/PropertySelection/SPropertySelectionComboButton.h b/Engine/Plugins/Developer/Concert/ConcertApp/MultiUserClient/Source/MultiUserClient/Private/Widgets/ActiveSession/Replication/Client/PropertySelection/SPropertySelectionComboButton.h
new file mode 100644
--- /dev/null
+++ b/Engine/Plugins/Developer/Concert/ConcertApp/MultiUserClient/Source/MultiUserClient/Private/Widgets/ActiveSession/Replication/Client/PropertySelection/SPropertySelectionComboButton.h
@@ -0,0 +1,56 @@
+﻿// Copyright Epic Games, Inc. All Rights Reserved.
+
+#pragma once
+
+#include "Model/Item/SourceModelBuilders.h"
+#include "SelectionDelegates.h"
+
+#include "Delegates/Delegate.h"
+#include "HAL/Platform.h"
+#include "Widgets/DeclarativeSyntaxSupport.h"
+#include "Widgets/Input/SComboButton.h"
+
+namespace UE::ConcertSharedSlate { struct FObjectGroup; }
+
+namespace UE::MultiUserClient
+{
+	class FRootPropertySourceModel;
+	class FUserPropertySelector;
+	struct FUserSelectableProperty;
+	
+	/**
+	 * This combo button is shown to the left of the search bar in the bottom half of the replication UI.
+	 * It allows users to specify the properties they want to work on (i.e. these properties should be shown in the property view).
+	 */
+	class SPropertySelectionComboButton : public SCompoundWidget
+	{
+		using FModelBuilder = ConcertSharedSlate::FSourceModelBuilders<FUserSelectableProperty>;
+	public:
+
+		SLATE_BEGIN_ARGS(SPropertySelectionComboButton) {}
+			/** Gets the display string of objects in the drop-down */
+			SLATE_EVENT(FGetObjectDisplayString, GetObjectDisplayString)
+		SLATE_END_ARGS()
+
+		void Construct(const FArguments& InArgs, FUserPropertySelector& InPropertySelector UE_LIFETIMEBOUND);
+
+		/** Refreshes the properties that the user can select given the objects currently displayed in the bottom view. */
+		void RefreshSelectableProperties(TConstArrayView<ConcertSharedSlate::FObjectGroup> DisplayedObjectGroups) const;
+
+	private:
+
+		/** Manages the user selected properties */
+		FUserPropertySelector* PropertySelector = nullptr;
+		
+		/** Manages the button's content. */
+		TUniquePtr<FRootPropertySourceModel> PropertySourceModel;
+
+		/** @return The menu to display when the combo button is clicked */
+		TSharedRef<SWidget> MakeMenu() const;
+		FModelBuilder::FItemPickerArgs MakePickerArguments() const;
+		
+		/** Handles the user selecting an option in the combo button. */
+		void OnItemsSelected(TArray<FUserSelectableProperty> Properties) const;
+	};
+}
+
diff --git a/Engine/Plugins/Developer/Concert/ConcertApp/MultiUserClient/Source/MultiUserClient/Private/Widgets/ActiveSession/Replication/Client/PropertySelection/SelectionDelegates.h b/Engine/Plugins/Developer/Concert/ConcertApp/MultiUserClient/Source/MultiUserClient/Private/Widgets/ActiveSession/Replication/Client/PropertySelection/SelectionDelegates.h
new file mode 100644
--- /dev/null
+++ b/Engine/Plugins/Developer/Concert/ConcertApp/MultiUserClient/Source/MultiUserClient/Private/Widgets/ActiveSession/Replication/Client/PropertySelection/SelectionDelegates.h
@@ -0,0 +1,14 @@
+﻿// Copyright Epic Games, Inc. All Rights Reserved.
+
+#pragma once
+
+#include "Delegates/Delegate.h"
+#include "UObject/SoftObjectPtr.h"
+
+class FText;
+
+namespace UE::MultiUserClient
+{
+	DECLARE_DELEGATE_RetVal_OneParam(FText, FGetObjectDisplayString, const TSoftObjectPtr<>&)
+}
+
diff --git a/Engine/Plugins/Developer/Concert/ConcertSync/ConcertSyncCore/Source/ConcertSyncCore/Private/Replication/Data/ConcertPropertySelection.cpp b/Engine/Plugins/Developer/Concert/ConcertSync/ConcertSyncCore/Source/ConcertSyncCore/Private/Replication/Data/ConcertPropertySelection.cpp
--- a/Engine/Plugins/Developer/Concert/ConcertSync/ConcertSyncCore/Source/ConcertSyncCore/Private/Replication/Data/ConcertPropertySelection.cpp
+++ b/Engine/Plugins/Developer/Concert/ConcertSync/ConcertSyncCore/Source/ConcertSyncCore/Private/Replication/Data/ConcertPropertySelection.cpp
@@ -182,6 +182,18 @@ FConcertPropertyChain FConcertPropertyChain::GetParent() const
     return Result;
 }
 
+FConcertPropertyChain FConcertPropertyChain::GetRootParent() const
+{
+	if (PathToProperty.Num() <= 1)
+	{
+		return {};
+	}
+	
+	FConcertPropertyChain Result = *this;
+	Result.PathToProperty = { PathToProperty[0] };
+	return Result;
+}
+
 FString FConcertPropertyChain::ToString(EToStringMethod Method) const
 {
     switch (Method)
diff --git a/Engine/Plugins/Developer/Concert/ConcertSync/ConcertSyncCore/Source/ConcertSyncCore/Public/Replication/Data/ConcertPropertySelection.h b/Engine/Plugins/Developer/Concert/ConcertSync/ConcertSyncCore/Source/ConcertSyncCore/Public/Replication/Data/ConcertPropertySelection.h
--- a/Engine/Plugins/Developer/Concert/ConcertSync/ConcertSyncCore/Source/ConcertSyncCore/Public/Replication/Data/ConcertPropertySelection.h
+++ b/Engine/Plugins/Developer/Concert/ConcertSync/ConcertSyncCore/Source/ConcertSyncCore/Public/Replication/Data/ConcertPropertySelection.h
@@ -52,7 +52,11 @@ struct CONCERTSYNCCORE_API FConcertPropertyChain
 	/** @return Attempts to resolve this property given the class */
 	FProperty* ResolveProperty(const UStruct& Class, bool bLogOnFail = true) const;
 
+	/** @return The property immediately before the current one in the chain */
 	FConcertPropertyChain GetParent() const;
+	/** @return The root-most parent in the chain. */
+	FConcertPropertyChain GetRootParent() const;
+	/** @return The property path */
 	const TArray<FName>& GetPathToProperty() const { return PathToProperty; }
 
 	enum class EToStringMethod
diff --git a/Engine/Plugins/Developer/Concert/ConcertUI/ConcertClientSharedSlate/Source/ConcertClientSharedSlate/Private/Replication/Editor/Model/Object/EditorObjectNameModel.cpp b/Engine/Plugins/Developer/Concert/ConcertUI/ConcertClientSharedSlate/Source/ConcertClientSharedSlate/Private/Replication/Editor/Model/Object/EditorObjectNameModel.cpp
--- a/Engine/Plugins/Developer/Concert/ConcertUI/ConcertClientSharedSlate/Source/ConcertClientSharedSlate/Private/Replication/Editor/Model/Object/EditorObjectNameModel.cpp
+++ b/Engine/Plugins/Developer/Concert/ConcertUI/ConcertClientSharedSlate/Source/ConcertClientSharedSlate/Private/Replication/Editor/Model/Object/EditorObjectNameModel.cpp
@@ -24,7 +24,8 @@ namespace UE::ConcertClientSharedSlate
 				const UObject* Object = SubobjectData->FindComponentInstanceInActor(&OwningActor);
 				if (Object == &Subbject)
 				{
-					return FText::FromString(SubobjectData->GetDisplayString());
+					constexpr bool bShowNativeComponentNames = false;
+					return FText::FromString(SubobjectData->GetDisplayString(bShowNativeComponentNames));
 				}
 			}
 
diff --git a/Engine/Plugins/Developer/Concert/ConcertUI/ConcertClientSharedSlate/Source/ConcertClientSharedSlate/Public/Replication/Editor/Model/PropertySource/ReplicatablePropertySource.h b/Engine/Plugins/Developer/Concert/ConcertUI/ConcertClientSharedSlate/Source/ConcertClientSharedSlate/Public/Replication/Editor/Model/PropertySource/ReplicatablePropertySource.h
--- a/Engine/Plugins/Developer/Concert/ConcertUI/ConcertClientSharedSlate/Source/ConcertClientSharedSlate/Public/Replication/Editor/Model/PropertySource/ReplicatablePropertySource.h
+++ b/Engine/Plugins/Developer/Concert/ConcertUI/ConcertClientSharedSlate/Source/ConcertClientSharedSlate/Public/Replication/Editor/Model/PropertySource/ReplicatablePropertySource.h
@@ -17,7 +17,7 @@ namespace UE::ConcertClientSharedSlate
 		{}
 
 		//~ Begin IPropertySource Interface
-		virtual void EnumerateProperties(TFunctionRef<EBreakBehavior(ConcertSharedSlate::FPropertyInfo&& Property)> Delegate) const override
+		virtual void EnumerateProperties(TFunctionRef<EBreakBehavior(const ConcertSharedSlate::FPropertyInfo& Property)> Delegate) const override
 		{
 			if (!Class.IsValid())
 			{
@@ -26,7 +26,7 @@ namespace UE::ConcertClientSharedSlate
 			
 			ConcertSyncCore::PropertyChain::ForEachReplicatableConcertProperty(*Class.Get(), [this, &Delegate](FConcertPropertyChain&& Property)
 			{
-				return Delegate(ConcertSharedSlate::FPropertyInfo(MoveTemp(Property)));
+				return Delegate(ConcertSharedSlate::FPropertyInfo(Property));
 			});
 		}
 		//~ End IPropertySource Interface
diff --git a/Engine/Plugins/Developer/Concert/ConcertUI/ConcertSharedSlate/Source/ConcertSharedSlate/Private/Replication/Editor/View/ObjectEditor/SBaseReplicationStreamEditor.cpp b/Engine/Plugins/Developer/Concert/ConcertUI/ConcertSharedSlate/Source/ConcertSharedSlate/Private/Replication/Editor/View/ObjectEditor/SBaseReplicationStreamEditor.cpp
--- a/Engine/Plugins/Developer/Concert/ConcertUI/ConcertSharedSlate/Source/ConcertSharedSlate/Private/Replication/Editor/View/ObjectEditor/SBaseReplicationStreamEditor.cpp
+++ b/Engine/Plugins/Developer/Concert/ConcertUI/ConcertSharedSlate/Source/ConcertSharedSlate/Private/Replication/Editor/View/ObjectEditor/SBaseReplicationStreamEditor.cpp
@@ -154,7 +154,7 @@ namespace UE::ConcertSharedSlate
 			Algo::TransformIf(AddedObjects, TopLevelObjects,
 				[this, &AddedObjects](const UObject* Object)
 				{
-					const TOptional<IObjectHierarchyModel::FParentInfo> ParentInfo = ObjectHierarchy->GetParentInfo(Object);
+					const TOptional<IObjectHierarchyModel::FParentInfo> ParentInfo = ObjectHierarchy ? ObjectHierarchy->GetParentInfo(Object) : TOptional<IObjectHierarchyModel::FParentInfo>{};
 					const bool bIsTopOfHierarchy = !ParentInfo || !AddedObjects.ContainsByPredicate([&ParentInfo](UObject* AddedObject){ return AddedObject == ParentInfo->Parent; });
 					return bIsTopOfHierarchy;
 				},
@@ -242,7 +242,7 @@ namespace UE::ConcertSharedSlate
 			TAttribute<FText>::CreateLambda([this](){ return GetEditingDisabledText(); }),
 			FSlateIcon(),
 			FUIAction(
-				FExecuteAction::CreateSP(this, &SBaseReplicationStreamEditor::OnDeleteObjects_PassByValue, ReplicationViewer->GetSelectedOutlinerObjects()),
+				FExecuteAction::CreateSP(this, &SBaseReplicationStreamEditor::OnDeleteObjects_PassByValue, ReplicationViewer->GetSelectedObjectItems()),
 				FCanExecuteAction::CreateLambda([this]() { return !IsEditingDisabled(); })
 				),
 			NAME_None,
@@ -250,7 +250,7 @@ namespace UE::ConcertSharedSlate
 		);
 
 		TArray<TSoftObjectPtr<>> SelectedObjects;
-		Algo::Transform(ReplicationViewer->GetSelectedOutlinerObjects(), SelectedObjects, [](const TSharedPtr<FReplicatedObjectData>& Data){ return Data->GetObjectPtr(); });
+		Algo::Transform(ReplicationViewer->GetSelectedObjectItems(), SelectedObjects, [](const TSharedPtr<FReplicatedObjectData>& Data){ return Data->GetObjectPtr(); });
 		OnExtendObjectsContextMenuDelegate.ExecuteIfBound(MenuBuilder, SelectedObjects);
 		
 		return MenuBuilder.MakeWidget();
@@ -261,7 +261,7 @@ namespace UE::ConcertSharedSlate
 		using namespace ConcertSharedSlate;
 		
 		// Context menu generation is only supported for single items
-		const TArray<TSharedPtr<FReplicatedObjectData>> SelectedObjects = ReplicationViewer->GetSelectedOutlinerObjects();
+		const TArray<TSharedPtr<FReplicatedObjectData>> SelectedObjects = ReplicationViewer->GetSelectedObjectItems();
 		if (SelectedObjects.Num() != 1)
 		{
 			return;
diff --git a/Engine/Plugins/Developer/Concert/ConcertUI/ConcertSharedSlate/Source/ConcertSharedSlate/Private/Replication/Editor/View/ObjectViewer/SReplicatedPropertyView.cpp b/Engine/Plugins/Developer/Concert/ConcertUI/ConcertSharedSlate/Source/ConcertSharedSlate/Private/Replication/Editor/View/ObjectViewer/SReplicatedPropertyView.cpp
--- a/Engine/Plugins/Developer/Concert/ConcertUI/ConcertSharedSlate/Source/ConcertSharedSlate/Private/Replication/Editor/View/ObjectViewer/SReplicatedPropertyView.cpp
+++ b/Engine/Plugins/Developer/Concert/ConcertUI/ConcertSharedSlate/Source/ConcertSharedSlate/Private/Replication/Editor/View/ObjectViewer/SReplicatedPropertyView.cpp
@@ -19,9 +19,7 @@ namespace UE::ConcertSharedSlate
 		PropertyAssignmentView = MoveTemp(InPropertyAssignmentView);
 		PropertiesModel = MoveTemp(InPropertiesModel);
 		
-		GetSelectedRootObjectsDelegate = InArgs._GetSelectedRootObjects;
 		GetObjectClassDelegate = InArgs._GetObjectClass;
-		check(GetSelectedRootObjectsDelegate.IsBound() && GetObjectClassDelegate.IsBound());
 		
 		ChildSlot
 		[
@@ -29,10 +27,9 @@ namespace UE::ConcertSharedSlate
 		];
 	}
 	
-	void SReplicatedPropertyView::RefreshPropertyData()
+	void SReplicatedPropertyView::RefreshPropertyData(const TArray<TSoftObjectPtr<>>& SelectedObjects)
 	{
 		SCOPED_CONCERT_TRACE(RefreshPropertyData);
-		TArray<TSoftObjectPtr<>> SelectedObjects = GetObjectsSelectedForPropertyEditing();
 		if (SelectedObjects.IsEmpty())
 		{
 			SetPropertyContent(EReplicatedPropertyContent::NoSelection);
@@ -51,16 +48,6 @@ namespace UE::ConcertSharedSlate
 		SetPropertyContent(EReplicatedPropertyContent::Properties);
 	}
 
-	TArray<TSoftObjectPtr<>> SReplicatedPropertyView::GetObjectsSelectedForPropertyEditing() const
-	{
-		TArray<TSoftObjectPtr<>> Result;
-		Algo::Transform(GetSelectedRootObjectsDelegate.Execute(), Result, [](const TSharedPtr<FReplicatedObjectData>& ObjectData)
-		{
-			return ObjectData->GetObjectPtr();
-		});
-		return Result;
-	}
-
 	TSharedRef<SWidget> SReplicatedPropertyView::CreatePropertiesView(const FArguments& InArgs)
 	{
 		return SAssignNew(PropertyContent, SWidgetSwitcher)
diff --git a/Engine/Plugins/Developer/Concert/ConcertUI/ConcertSharedSlate/Source/ConcertSharedSlate/Private/Replication/Editor/View/ObjectViewer/SReplicatedPropertyView.h b/Engine/Plugins/Developer/Concert/ConcertUI/ConcertSharedSlate/Source/ConcertSharedSlate/Private/Replication/Editor/View/ObjectViewer/SReplicatedPropertyView.h
--- a/Engine/Plugins/Developer/Concert/ConcertUI/ConcertSharedSlate/Source/ConcertSharedSlate/Private/Replication/Editor/View/ObjectViewer/SReplicatedPropertyView.h
+++ b/Engine/Plugins/Developer/Concert/ConcertUI/ConcertSharedSlate/Source/ConcertSharedSlate/Private/Replication/Editor/View/ObjectViewer/SReplicatedPropertyView.h
@@ -21,14 +21,9 @@ namespace UE::ConcertSharedSlate
 	class SReplicatedPropertyView : public SCompoundWidget
 	{
 	public:
-		
-		DECLARE_DELEGATE_RetVal(TArray<TSharedPtr<FReplicatedObjectData>>, FGetSelectedRootObjects)
 
 		SLATE_BEGIN_ARGS(SReplicatedPropertyView)
 		{}
-			/** Gets the root objects selected in the object outliner. */
-			SLATE_EVENT(FGetSelectedRootObjects, GetSelectedRootObjects)
-
 			/** Gets the class for the object since the object may not be in the model. */
 			SLATE_EVENT(FGetObjectClass, GetObjectClass)
 		
@@ -37,12 +32,11 @@ namespace UE::ConcertSharedSlate
 		SLATE_END_ARGS()
 
 		void Construct(const FArguments& InArgs, TSharedRef<IPropertyAssignmentView> InPropertyAssignmentView, TSharedRef<IReplicationStreamModel> InPropertiesModel);
-		
-		void RefreshPropertyData();
+
+		/** Updates the displayed properties */
+		void RefreshPropertyData(const TArray<TSoftObjectPtr<>>& SelectedObjects);
 		/** Requests that the given column be resorted, if it currently affects the row sorting. */
 		void RequestResortForColumn(const FName& ColumnId) const { PropertyAssignmentView->RequestResortForColumn(ColumnId); }
-		
-		TArray<TSoftObjectPtr<>> GetObjectsSelectedForPropertyEditing() const;
 
 	private:
 		
@@ -64,8 +58,6 @@ namespace UE::ConcertSharedSlate
 		/** Determines the content displayed for PropertyArea. */
 		TSharedPtr<SWidgetSwitcher> PropertyContent;
 		
-		/** Gets the root objects selected in the object outliner. */
-		FGetSelectedRootObjects GetSelectedRootObjectsDelegate;
 		/** Gets the class for the object since the object may not be in the model. */
 		FGetObjectClass GetObjectClassDelegate;
 
diff --git a/Engine/Plugins/Developer/Concert/ConcertUI/ConcertSharedSlate/Source/ConcertSharedSlate/Private/Replication/Editor/View/ObjectViewer/SReplicationStreamViewer.cpp b/Engine/Plugins/Developer/Concert/ConcertUI/ConcertSharedSlate/Source/ConcertSharedSlate/Private/Replication/Editor/View/ObjectViewer/SReplicationStreamViewer.cpp
--- a/Engine/Plugins/Developer/Concert/ConcertUI/ConcertSharedSlate/Source/ConcertSharedSlate/Private/Replication/Editor/View/ObjectViewer/SReplicationStreamViewer.cpp
+++ b/Engine/Plugins/Developer/Concert/ConcertUI/ConcertSharedSlate/Source/ConcertSharedSlate/Private/Replication/Editor/View/ObjectViewer/SReplicationStreamViewer.cpp
@@ -73,7 +73,12 @@ namespace UE::ConcertSharedSlate
 
 	TArray<TSoftObjectPtr<>> SReplicationStreamViewer::GetSelectedObjects() const
 	{
-		return PropertySection->GetObjectsSelectedForPropertyEditing();
+		TArray<TSoftObjectPtr<>> Result;
+		Algo::Transform(GetSelectedObjectItems(), Result, [](const TSharedPtr<FReplicatedObjectData>& Item)
+		{
+			return Item->GetObjectPtr();
+		});
+		return Result;
 	}
 
 	void SReplicationStreamViewer::SelectObjects(TConstArrayView<TSoftObjectPtr<>> Objects, bool bAtEndOfTick)
@@ -87,9 +92,9 @@ namespace UE::ConcertSharedSlate
 		}
 		
 		TArray<TSharedPtr<FReplicatedObjectData>> NewSelectedItems; 
-		Algo::TransformIf(AllObjectRowData, NewSelectedItems, [&Objects](const TSharedPtr<FReplicatedObjectData>& ObjectData)
+		Algo::TransformIf(AllObjectRowData, NewSelectedItems, [this, &Objects](const TSharedPtr<FReplicatedObjectData>& ObjectData)
 			{
-				return Objects.Contains(ObjectData->GetObjectPtr());
+				return Objects.Contains(ObjectData->GetObjectPtr()) && CanDisplayObject(ObjectData->GetObjectPtr());
 			},
 			[](const TSharedPtr<FReplicatedObjectData>& ObjectData){ return ObjectData; }
 		);
@@ -143,7 +148,7 @@ namespace UE::ConcertSharedSlate
 		}
 	}
 	
-	TArray<TSharedPtr<FReplicatedObjectData>> SReplicationStreamViewer::GetSelectedOutlinerObjects() const
+	TArray<TSharedPtr<FReplicatedObjectData>> SReplicationStreamViewer::GetSelectedObjectItems() const
 	{
 		TArray<TSharedPtr<FReplicatedObjectData>> SelectedItems = ReplicatedObjects->GetSelectedItems();
 		// Items may have been removed this tick. However, selected items may not have been updated yet because STreeView processes item changes at the end of tick. 
@@ -183,7 +188,10 @@ namespace UE::ConcertSharedSlate
 		if (bHasRequestedPropertyRefresh)
 		{
 			bHasRequestedPropertyRefresh = false;
-			RefreshPropertyData();
+			PropertySection->RefreshPropertyData(
+				// If we're about to change the selection, pass in those objects
+				PendingToSelect.IsEmpty() ? GetSelectedObjects() : PendingToSelect
+				);
 		}
 
 		if (!PendingToSelect.IsEmpty())
@@ -302,7 +310,6 @@ namespace UE::ConcertSharedSlate
 				.BodyContent()
 				[
 					SAssignNew(PropertySection, SReplicatedPropertyView, InArgs._PropertyAssignmentView.ToSharedRef(), PropertiesModel.ToSharedRef())
-					.GetSelectedRootObjects(this, &SReplicationStreamViewer::GetSelectedOutlinerObjects)
 					.GetObjectClass(this, &SReplicationStreamViewer::GetObjectClass)
 					.NameModel(InArgs._NameModel)
 				]
@@ -325,20 +332,15 @@ namespace UE::ConcertSharedSlate
 		// An alternative would be to change RefreshObjectData to be called with two variables ObjectsAdded and ObjectsRemoved.
 		IterateDisplayableObjects([this, &NewPathToObjectDataCache](const FSoftObjectPath& ObjectPath) mutable
 		{
-			TOptional<IObjectHierarchyModel::FParentInfo> ParentInfo = ObjectHierarchy->GetParentInfo(TSoftObjectPtr(ObjectPath));
-			const bool bIsActor = !ParentInfo; 
-			if (bIsActor || ShouldDisplayChildObject(ObjectPath, ParentInfo->Relationship))
-			{
-				const TSharedPtr<FReplicatedObjectData>* ExistingItem = PathToObjectDataCache.Find(ObjectPath);
-				ExistingItem = ExistingItem ? ExistingItem : NewPathToObjectDataCache.Find(ObjectPath);
-				const TSharedRef<FReplicatedObjectData> Item = ExistingItem
-					? ExistingItem->ToSharedRef()
-					: Private::AllocateObjectData(ObjectPath);
-				AllObjectRowData.AddUnique(Item);
-				NewPathToObjectDataCache.Emplace(ObjectPath, Item);
-				
-				BuildObjectHierarchyIfNeeded(Item, NewPathToObjectDataCache);
-			}
+			const TSharedPtr<FReplicatedObjectData>* ExistingItem = PathToObjectDataCache.Find(ObjectPath);
+			ExistingItem = ExistingItem ? ExistingItem : NewPathToObjectDataCache.Find(ObjectPath);
+			const TSharedRef<FReplicatedObjectData> Item = ExistingItem
+				? ExistingItem->ToSharedRef()
+				: Private::AllocateObjectData(ObjectPath);
+			AllObjectRowData.AddUnique(Item);
+			NewPathToObjectDataCache.Emplace(ObjectPath, Item);
+			
+			BuildObjectHierarchyIfNeeded(Item, NewPathToObjectDataCache);
 			
 			return EBreakBehavior::Continue;
 		});
@@ -360,78 +362,49 @@ namespace UE::ConcertSharedSlate
 		// Case: RealModel contains only components but not the owning actor.
 		// In that case, we want the UI to still show the owning actor.
 		// We'll track this with these containers:
-		TSet<FSoftObjectPath> NonActors;
-		TSet<FSoftObjectPath> Actors;
+		TSet<FSoftObjectPath> AddedActors;
+		TSet<FSoftObjectPath> PendingActors;
 
-		// Here we only exposes objects that are top level so we skip components, etc.
-		EBreakBehavior BreakBehavior = EBreakBehavior::Continue;
-		PropertiesModel->ForEachReplicatedObject([this, &Delegate, &NonActors, &Actors, &BreakBehavior](const FSoftObjectPath& Object)
+		PropertiesModel->ForEachReplicatedObject([this, &Delegate, &AddedActors, &PendingActors](const FSoftObjectPath& Object)
 		{
+			if (const TOptional<FSoftObjectPath> OwningActor = ObjectUtils::GetActorOf(Object)
+				; OwningActor && CanDisplayObject(*OwningActor))
+			{
+				PendingActors.Add(*OwningActor);
+			}
 			if (!CanDisplayObject(Object))
 			{
 				return EBreakBehavior::Continue;
 			}
 			
 			Delegate(Object);
-			
 			if (ObjectUtils::IsActor(Object))
 			{
-				Actors.Add(Object);
-			}
-			else
-			{
-				NonActors.Add(Object);
+				AddedActors.Add(Object);
 			}
 			
 			return EBreakBehavior::Continue;
 		});
 
-		// No additional objects to show
-		if (BreakBehavior == EBreakBehavior::Break)
+		// Now determine the actors that are not in PropertiesModel but that need to be shown because its subobjects that are in PropertiesModel
+		for (const FSoftObjectPath& PendingActor : PendingActors)
 		{
-			return;
-		}
-
-		// Now determine the top level objects that are not in RealModel but that need to be shown for subobjects that are in RealModel
-		for (const FSoftObjectPath& NonTopLevelObject : NonActors)
-		{
-			const TOptional<FSoftObjectPath> TopLevelObject = ObjectUtils::GetActorOf(NonTopLevelObject);
-			if (!TopLevelObject || !CanDisplayObject(*TopLevelObject))
+			if (!AddedActors.Contains(PendingActor))
 			{
-				continue;
-			}
-			
-			if (!Actors.Contains(*TopLevelObject))
-			{
-				Delegate(*TopLevelObject);
-				
-				// Avoid calling Delegate again for TopLevelObject
-				Actors.Add(*TopLevelObject);
+				Delegate(PendingActor);
+				AddedActors.Add(PendingActor);
 			}
 		}
 	}
 
-	void SReplicationStreamViewer::RefreshPropertyData()
-	{
-		PropertySection->RefreshPropertyData();
-	}
-
 	void SReplicationStreamViewer::BuildRootObjectRowData()
 	{
+		RootObjectRowData.Empty();
+		
 		TSet<TSharedPtr<FReplicatedObjectData>> NonRootNodes;
 		for (const TSharedPtr<FReplicatedObjectData>& Node : AllObjectRowData)
 		{
-			GetObjectRowChildren(Node, [&NonRootNodes](TSharedPtr<FReplicatedObjectData> Child)
-			{
-				NonRootNodes.Add(MoveTemp(Child));
-			});
-		}
-
-		// Make RootObjectRowData only contain those nodes which were not listed as children 
-		RootObjectRowData.Empty(NonRootNodes.Num());
-		for (const TSharedPtr<FReplicatedObjectData>& Node : AllObjectRowData)
-		{
-			if (!NonRootNodes.Contains(Node))
+			if (ObjectUtils::IsActor(Node->GetObjectPath()))
 			{
 				RootObjectRowData.Add(Node);
 			}
@@ -474,9 +447,9 @@ namespace UE::ConcertSharedSlate
 		
 		const TSoftObjectPtr ActorPtr(OwningActor);
 		AddItem(ActorPtr);
-		ObjectHierarchy->ForEachChildRecursive(ActorPtr, [this, &AddItem](const TSoftObjectPtr<>&, const TSoftObjectPtr<>& ChildObject, EChildRelationship Relationship)
+		ObjectHierarchy->ForEachChildRecursive(ActorPtr, [this, &AddItem](const TSoftObjectPtr<>&, const TSoftObjectPtr<>& ChildObject, EChildRelationship)
 		{
-			if (ShouldDisplayChildObject(ChildObject.GetUniqueID(), Relationship))
+			if (CanDisplayObject(ChildObject))
 			{
 				AddItem(ChildObject);
 			}
@@ -503,9 +476,14 @@ namespace UE::ConcertSharedSlate
 			return EBreakBehavior::Continue;
 		});
 	}
-	bool SReplicationStreamViewer::CanDisplayObject(const FSoftObjectPath& ObjectPath) const
+	bool SReplicationStreamViewer::CanDisplayObject(const TSoftObjectPtr<>& Object) const
 	{
-		return !ShouldDisplayObjectDelegate.IsBound() || ShouldDisplayObjectDelegate.Execute(ObjectPath);
+		const TOptional<IObjectHierarchyModel::FParentInfo> ParentInfo = ObjectHierarchy ? ObjectHierarchy->GetParentInfo(Object) : TOptional<IObjectHierarchyModel::FParentInfo>{};
+		// If no hierarchy was provided during construction, only show actors. If hierarchy provided, check whether this type of subobject is allowed.
+		const bool bCanShowWithinHierarchy = (ParentInfo && ShouldDisplayObjectRelation(ParentInfo->Relationship))
+			|| ObjectUtils::IsActor(Object.GetUniqueID());
+		const bool bDidDelegateAllow = !ShouldDisplayObjectDelegate.IsBound() || ShouldDisplayObjectDelegate.Execute(Object.GetUniqueID());
+		return bCanShowWithinHierarchy && bDidDelegateAllow;
 	}
 
 	bool SReplicationStreamViewer::ShouldDisplayObjectRelation(EChildRelationship Relationship) const
diff --git a/Engine/Plugins/Developer/Concert/ConcertUI/ConcertSharedSlate/Source/ConcertSharedSlate/Private/Replication/Editor/View/ObjectViewer/SReplicationStreamViewer.h b/Engine/Plugins/Developer/Concert/ConcertUI/ConcertSharedSlate/Source/ConcertSharedSlate/Private/Replication/Editor/View/ObjectViewer/SReplicationStreamViewer.h
--- a/Engine/Plugins/Developer/Concert/ConcertUI/ConcertSharedSlate/Source/ConcertSharedSlate/Private/Replication/Editor/View/ObjectViewer/SReplicationStreamViewer.h
+++ b/Engine/Plugins/Developer/Concert/ConcertUI/ConcertSharedSlate/Source/ConcertSharedSlate/Private/Replication/Editor/View/ObjectViewer/SReplicationStreamViewer.h
@@ -107,7 +107,7 @@ namespace UE::ConcertSharedSlate
 		void ExpandObjects(TConstArrayView<TSoftObjectPtr<>> Objects, bool bRecursive, bool bAtEndOfTick = false);
 
 		/** @return Gets the root objects selected in the outliner; the subobject view chooses which of these objects (or their subobjects) end up in GetSelectedObjectShowingProperties. */
-		TArray<TSharedPtr<FReplicatedObjectData>> GetSelectedOutlinerObjects() const;
+		TArray<TSharedPtr<FReplicatedObjectData>> GetSelectedObjectItems() const;
 
 		//~ Begin SWidget Interface
 		virtual void Tick(const FGeometry& AllottedGeometry, const double InCurrentTime, const float InDeltaTime) override;
@@ -160,9 +160,6 @@ namespace UE::ConcertSharedSlate
 		/** Lists all objects that need an object in the top-view. */
 		void IterateDisplayableObjects(TFunctionRef<void(const FSoftObjectPath& Object)> Delegate) const;
 		
-		/** Regenerates property row data in the bottom videw */
-		void RefreshPropertyData();
-		
 		/** Sets RootObjectRowData to all non-root nodes from ObjectRowData. */
 		void BuildRootObjectRowData();
 
@@ -180,13 +177,10 @@ namespace UE::ConcertSharedSlate
 			RequestObjectDataRefresh();
 			RequestPropertyDataRefresh();
 		}
-		bool ShouldDisplayChildObject(const FSoftObjectPath& Object, EChildRelationship Relationship) const
-		{
-			return CanDisplayObject(Object) && ShouldDisplayObjectRelation(Relationship);
-		}
 		
 		/** Invokes the ShouldDisplayObjectDelegate to determine whether ObjectPath should be displayed. */
-		bool CanDisplayObject(const FSoftObjectPath& ObjectPath) const;
+		bool CanDisplayObject(const TSoftObjectPtr<>& Object) const;
+		bool CanDisplayObject(const FSoftObjectPath& ObjectPath) const { return CanDisplayObject(TSoftObjectPtr<>(ObjectPath)); }
 		/** Whether the view options allow this type of relationship to be shown. */
 		bool ShouldDisplayObjectRelation(EChildRelationship Relationship) const;
 		
diff --git a/Engine/Plugins/Developer/Concert/ConcertUI/ConcertSharedSlate/Source/ConcertSharedSlate/Private/Replication/Editor/View/Property/SMultiObjectAssignment.cpp b/Engine/Plugins/Developer/Concert/ConcertUI/ConcertSharedSlate/Source/ConcertSharedSlate/Private/Replication/Editor/View/Property/SMultiObjectAssignment.cpp
--- a/Engine/Plugins/Developer/Concert/ConcertUI/ConcertSharedSlate/Source/ConcertSharedSlate/Private/Replication/Editor/View/Property/SMultiObjectAssignment.cpp
+++ b/Engine/Plugins/Developer/Concert/ConcertUI/ConcertSharedSlate/Source/ConcertSharedSlate/Private/Replication/Editor/View/Property/SMultiObjectAssignment.cpp
@@ -114,17 +114,15 @@ namespace UE::ConcertSharedSlate
 	void SMultiObjectAssignment::RefreshData(const TArray<TSoftObjectPtr<>>& Objects, const IReplicationStreamModel& Model)
 	{
 		TArray<FPropertyAssignmentEntry> Entries;
+		// The root objects are always "displayed" for the purposes of GetDisplayedGroups
+		DisplayedGroups = { FObjectGroup{ Objects } };
+		
 		auto[RootEntry, bRootSharesClass] = BuildAssignmentEntry(Objects, Model);
-		if (RootEntry.PropertiesToDisplay.IsEmpty() || !ensureMsgf(bRootSharesClass, TEXT("Objects do not share the same class. Investigate invalid call.")))
+		if (!RootEntry.PropertiesToDisplay.IsEmpty() && ensureMsgf(bRootSharesClass, TEXT("Objects do not share the same class. Investigate invalid call.")))
 		{
-			PreviousSelectedObjects = Objects;
-			TreeView->RefreshPropertyData({}, false);
-			OnObjectGroupsChangedDelegate.Broadcast();
-			return;
+			Entries.Emplace(MoveTemp(RootEntry));
 		}
 
-		DisplayedGroups = { FObjectGroup{ Objects } };
-		Entries.Emplace(MoveTemp(RootEntry));
 		if (ObjectHierarchy)
 		{
 			using namespace Private;
@@ -140,21 +138,30 @@ namespace UE::ConcertSharedSlate
 				}
 			});
 		}
+
+		if (!Entries.IsEmpty())
+		{
+			const bool bObjectsHaveChanged = PreviousSelectedObjects != Objects;
+			PreviousSelectedObjects = Objects;
 		
-		const bool bObjectsHaveChanged = PreviousSelectedObjects == Objects;
-		PreviousSelectedObjects = Objects;
-		
-		// If the objects have changed, the classes may share properties.
-		// In that case, below we'd reuse the item pointer, which would cause the tree view to re-use the old row widgets.
-		// However, we must regenerate all column widgets since they may be referencing the object the row was originally built for. So they'd display the state of the previous object still!
-		// Example: Assign property combo-box in Multi-User All Clients view displays who has the property assigned.
-		// Note: If the objects did not change, we definitely want to reuse item pointers since otherwise the user row selection is reset.
-		const bool bCanReusePropertyData = bObjectsHaveChanged;
+			// If the objects have changed, the classes may share properties.
+			// In that case, below we'd reuse the item pointer, which would cause the tree view to re-use the old row widgets.
+			// However, we must regenerate all column widgets since they may be referencing the object the row was originally built for. So they'd display the state of the previous object still!
+			// Example: Assign property combo-box in Multi-User All Clients view displays who has the property assigned.
+			// Note: If the objects did not change, we definitely want to reuse item pointers since otherwise the user row selection is reset.
+			const bool bCanReusePropertyData = !bObjectsHaveChanged;
 		
-		TreeView->RefreshPropertyData(Entries, bCanReusePropertyData);
+			TreeView->RefreshPropertyData(Entries, bCanReusePropertyData);
 
-		if (!bObjectsHaveChanged)
+			if (!bObjectsHaveChanged)
+			{
+				OnObjectGroupsChangedDelegate.Broadcast();
+			}
+		}
+		else
 		{
+			PreviousSelectedObjects.Empty();
+			TreeView->RefreshPropertyData({}, false);
 			OnObjectGroupsChangedDelegate.Broadcast();
 		}
 	}
diff --git a/Engine/Plugins/Developer/Concert/ConcertUI/ConcertSharedSlate/Source/ConcertSharedSlate/Private/Replication/Editor/View/Property/SPerObjectPropertyAssignment.cpp b/Engine/Plugins/Developer/Concert/ConcertUI/ConcertSharedSlate/Source/ConcertSharedSlate/Private/Replication/Editor/View/Property/SPerObjectPropertyAssignment.cpp
--- a/Engine/Plugins/Developer/Concert/ConcertUI/ConcertSharedSlate/Source/ConcertSharedSlate/Private/Replication/Editor/View/Property/SPerObjectPropertyAssignment.cpp
+++ b/Engine/Plugins/Developer/Concert/ConcertUI/ConcertSharedSlate/Source/ConcertSharedSlate/Private/Replication/Editor/View/Property/SPerObjectPropertyAssignment.cpp
@@ -21,6 +21,8 @@ namespace UE::ConcertSharedSlate
 
 	void SPerObjectPropertyAssignment::RefreshData(const TArray<TSoftObjectPtr<>>& Objects, const IReplicationStreamModel& Model)
 	{
+		DisplayedGroups = { FObjectGroup{ Objects } };
+		
 		FPropertyAssignmentEntry AssignmentEntry { .ContextObjects = Objects };
 		TSet<FConcertPropertyChain>& Properties = AssignmentEntry.PropertiesToDisplay;
 		FSoftClassPath& ClassPath = AssignmentEntry.Class;
@@ -47,7 +49,7 @@ namespace UE::ConcertSharedSlate
 		
 		if (bHasClassPath && !Properties.IsEmpty())
 		{
-			const bool bObjectsHaveChanged = PreviousSelectedObjects == Objects;
+			const bool bObjectsHaveChanged = PreviousSelectedObjects != Objects;
 			PreviousSelectedObjects = Objects;
 			
 			// If the objects have changed, the classes may share properties.
@@ -55,7 +57,7 @@ namespace UE::ConcertSharedSlate
 			// However, we must regenerate all column widgets since they may be referencing the object the row was originally built for. So they'd display the state of the previous object still!
 			// Example: Assign property combo-box in Multi-User All Clients view displays who has the property assigned.
 			// Note: If the objects did not change, we definitely want to reuse item pointers since otherwise the user row selection is reset.
-			const bool bCanReusePropertyData = bObjectsHaveChanged;
+			const bool bCanReusePropertyData = !bObjectsHaveChanged;
 			TreeView->RefreshPropertyData({ AssignmentEntry }, bCanReusePropertyData);
 
 			if (!bObjectsHaveChanged)
@@ -65,7 +67,7 @@ namespace UE::ConcertSharedSlate
 		}
 		else
 		{
-			PreviousSelectedObjects = Objects;
+			PreviousSelectedObjects.Empty();
 			TreeView->RefreshPropertyData({}, false);
 			OnObjectGroupsChangedDelegate.Broadcast();
 		}
diff --git a/Engine/Plugins/Developer/Concert/ConcertUI/ConcertSharedSlate/Source/ConcertSharedSlate/Private/Replication/Editor/View/Property/SPerObjectPropertyAssignment.h b/Engine/Plugins/Developer/Concert/ConcertUI/ConcertSharedSlate/Source/ConcertSharedSlate/Private/Replication/Editor/View/Property/SPerObjectPropertyAssignment.h
--- a/Engine/Plugins/Developer/Concert/ConcertUI/ConcertSharedSlate/Source/ConcertSharedSlate/Private/Replication/Editor/View/Property/SPerObjectPropertyAssignment.h
+++ b/Engine/Plugins/Developer/Concert/ConcertUI/ConcertSharedSlate/Source/ConcertSharedSlate/Private/Replication/Editor/View/Property/SPerObjectPropertyAssignment.h
@@ -32,7 +32,7 @@ namespace UE::ConcertSharedSlate
 		virtual void RequestRefilter() const override { return TreeView->RequestRefilter(); }
 		virtual void RequestResortForColumn(const FName& ColumnId) override { return TreeView->RequestRefilter(); }
 		virtual TSharedRef<SWidget> GetWidget() override { return SharedThis(this); }
-		virtual TArray<FObjectGroup> GetDisplayedGroups() const override { return { FObjectGroup{ PreviousSelectedObjects } }; }
+		virtual TArray<FObjectGroup> GetDisplayedGroups() const override { return DisplayedGroups; }
 		virtual FOnSelectionChanged& OnObjectGroupsChanged() override { return OnObjectGroupsChangedDelegate; }
 		//~ End IPropertyAssignmentView Interface
 
@@ -42,6 +42,8 @@ namespace UE::ConcertSharedSlate
 		TSharedPtr<IPropertyTreeView> TreeView;
 		/** Used to determine whether to rebuild the entire property data. */
 		TArray<TSoftObjectPtr<>> PreviousSelectedObjects;
+		/** Cached value for GetDisplayedGroups. */
+		TArray<FObjectGroup> DisplayedGroups;
 		
 		/** Optional. If specified, displays the properties of this model instead of those assigned in the stream. */
 		TSharedPtr<IPropertySourceProcessor> OptionalPropertySource;
diff --git a/Engine/Plugins/Developer/Concert/ConcertUI/ConcertSharedSlate/Source/ConcertSharedSlate/Private/Replication/Editor/View/Tree/SReplicationTreeView.h b/Engine/Plugins/Developer/Concert/ConcertUI/ConcertSharedSlate/Source/ConcertSharedSlate/Private/Replication/Editor/View/Tree/SReplicationTreeView.h
--- a/Engine/Plugins/Developer/Concert/ConcertUI/ConcertSharedSlate/Source/ConcertSharedSlate/Private/Replication/Editor/View/Tree/SReplicationTreeView.h
+++ b/Engine/Plugins/Developer/Concert/ConcertUI/ConcertSharedSlate/Source/ConcertSharedSlate/Private/Replication/Editor/View/Tree/SReplicationTreeView.h
@@ -468,6 +468,9 @@ namespace UE::ConcertSharedSlate
 	{
 		TSharedPtr<SVerticalBox> VerticalBox;
 		
+		const auto GetTreeVisibility = [this](){ return AllRootItems->IsEmpty() || FilteredRootItems.IsEmpty() ? EVisibility::Collapsed : EVisibility::Visible; };
+		const auto GetWarningVisibility = [this](){ return AllRootItems->IsEmpty() || FilteredRootItems.IsEmpty() ? EVisibility::Visible : EVisibility::Collapsed; };
+		
 		TSharedRef<SWidget> Result = SNew(SBorder)
 			.BorderImage(FAppStyle::Get().GetBrush("ToolPanel.GroupBorder"))
 			.BorderBackgroundColor(FSlateColor(FLinearColor(0.6, 0.6, 0.6)))
@@ -479,6 +482,7 @@ namespace UE::ConcertSharedSlate
 				.FillHeight(1.f)
 				[
 					SAssignNew(TreeView, STreeView<TSharedPtr<TItemType>>)
+					.Visibility_Lambda(GetTreeVisibility)
 					.OnGetChildren(this, &SReplicationTreeView::GetRowChildren)
 					.TreeItemsSource(&FilteredRootItems)
 					.OnGenerateRow(this, &SReplicationTreeView::OnGenerateRowWidget)
@@ -500,7 +504,7 @@ namespace UE::ConcertSharedSlate
 				[
 					SNew(SWidgetSwitcher)
 					.WidgetIndex_Lambda([this](){ return AllRootItems->IsEmpty() ? 1 : 0; })
-					.Visibility_Lambda([this](){ return AllRootItems->IsEmpty() || FilteredRootItems.IsEmpty() ? EVisibility::Visible : EVisibility::Collapsed; })
+					.Visibility_Lambda(GetWarningVisibility)
 					+SWidgetSwitcher::Slot() [ SNew(STextBlock).Text(LOCTEXT("AllFiltered", "All items are filtered.")) ]
 					+SWidgetSwitcher::Slot()
 					[
diff --git a/Engine/Plugins/Developer/Concert/ConcertUI/ConcertSharedSlate/Source/ConcertSharedSlate/Public/Replication/Editor/Model/Property/IPropertySource.h b/Engine/Plugins/Developer/Concert/ConcertUI/ConcertSharedSlate/Source/ConcertSharedSlate/Public/Replication/Editor/Model/Property/IPropertySource.h
--- a/Engine/Plugins/Developer/Concert/ConcertUI/ConcertSharedSlate/Source/ConcertSharedSlate/Public/Replication/Editor/Model/Property/IPropertySource.h
+++ b/Engine/Plugins/Developer/Concert/ConcertUI/ConcertSharedSlate/Source/ConcertSharedSlate/Public/Replication/Editor/Model/Property/IPropertySource.h
@@ -3,19 +3,20 @@
 #pragma once
 
 #include "Misc/EBreakBehavior.h"
-#include "Replication/Data/ConcertPropertySelection.h"
 
 #include "Templates/FunctionFwd.h"
 #include "UObject/SoftObjectPath.h"
 #include "UObject/SoftObjectPtr.h"
 
+struct FConcertPropertyChain;
+
 namespace UE::ConcertSharedSlate
 {
 	/** Wraps FConcertPropertyChain so it is easier to potentially change what IPropertySource lists in the future. */
 	struct FPropertyInfo
 	{
-		FConcertPropertyChain Property;
-		explicit FPropertyInfo(FConcertPropertyChain Property) : Property(MoveTemp(Property)) {}
+		const FConcertPropertyChain& Property;
+		explicit FPropertyInfo(const FConcertPropertyChain& Property) : Property(Property) {}
 	};
 	
 	/** Lists out a bunch of properties. */
@@ -24,7 +25,7 @@ namespace UE::ConcertSharedSlate
 	public:
 
 		/** Lists a bunch of properties. */
-		virtual void EnumerateProperties(TFunctionRef<EBreakBehavior(FPropertyInfo&& Property)> Delegate) const = 0;
+		virtual void EnumerateProperties(TFunctionRef<EBreakBehavior(const FPropertyInfo& Property)> Delegate) const = 0;
 		
 		virtual ~IPropertySource() = default;
 	};
diff --git a/Engine/Plugins/Developer/Concert/ConcertUI/ConcertSharedSlate/Source/ConcertSharedSlate/Public/Replication/Editor/View/Column/SelectionViewerColumns.h b/Engine/Plugins/Developer/Concert/ConcertUI/ConcertSharedSlate/Source/ConcertSharedSlate/Public/Replication/Editor/View/Column/SelectionViewerColumns.h
--- a/Engine/Plugins/Developer/Concert/ConcertUI/ConcertSharedSlate/Source/ConcertSharedSlate/Public/Replication/Editor/View/Column/SelectionViewerColumns.h
+++ b/Engine/Plugins/Developer/Concert/ConcertUI/ConcertSharedSlate/Source/ConcertSharedSlate/Public/Replication/Editor/View/Column/SelectionViewerColumns.h
@@ -53,7 +53,7 @@ namespace UE::ConcertSharedSlate::ReplicationColumns::TopLevel
 	/** Shows the object's type. */
 	CONCERTSHAREDSLATE_API FObjectColumnEntry TypeColumn(FGetObjectClass GetObjectClassDelegate);
 	/** Shows the number of widgets assigned to the object. */
-	CONCERTSHAREDSLATE_API FObjectColumnEntry NumPropertiesColumn(const IReplicationStreamModel& Model UE_LIFETIMEBOUND, ENumPropertiesFlags Flags = ENumPropertiesFlags::None);
+	CONCERTSHAREDSLATE_API FObjectColumnEntry NumPropertiesColumn(const IReplicationStreamModel& Model UE_LIFETIMEBOUND, ENumPropertiesFlags Flags = ENumPropertiesFlags::IncludeSubobjectCounts);
 }
 
 namespace UE::ConcertSharedSlate::ReplicationColumns::Property
