Fix deprecation warnings in concert shared slate.

diff --git a/Engine/Plugins/Developer/Concert/ConcertUI/ConcertSharedSlate/Source/ConcertSharedSlate/Public/Model/Item/IItemSourceModel.h b/Engine/Plugins/Developer/Concert/ConcertUI/ConcertSharedSlate/Source/ConcertSharedSlate/Public/Model/Item/IItemSourceModel.h
--- a/Engine/Plugins/Developer/Concert/ConcertUI/ConcertSharedSlate/Source/ConcertSharedSlate/Public/Model/Item/IItemSourceModel.h
+++ b/Engine/Plugins/Developer/Concert/ConcertUI/ConcertSharedSlate/Source/ConcertSharedSlate/Public/Model/Item/IItemSourceModel.h
@@ -91,6 +91,7 @@ namespace UE::ConcertSharedSlate
 		}
 		
 		/** Util that converts EnumerateSelectableItems into an array. */
+		PRAGMA_DISABLE_DEPRECATION_WARNINGS
 		TArray<TItemType> GetSelectableItems()
 		{
 			TArray<TItemType> Result;
@@ -101,6 +102,7 @@ namespace UE::ConcertSharedSlate
 			});
 			return Result;
 		}
+		PRAGMA_ENABLE_DEPRECATION_WARNINGS
 
 		virtual ~IItemSourceModel() = default;
 	};
