Fix and improve replicated property search:
- Fix child properties like "X" of RelativeLocation not showing up in search
- Fix searching by assigned client not working
- Fix expansion state being reset after search
- Remove highlight text from category rows

diff --git a/Engine/Plugins/Developer/Concert/ConcertApp/MultiUserClient/Source/MultiUserClient/Private/Widgets/ActiveSession/Replication/Client/Multi/Columns/MultiStreamColumns_AssignProperties.cpp b/Engine/Plugins/Developer/Concert/ConcertApp/MultiUserClient/Source/MultiUserClient/Private/Widgets/ActiveSession/Replication/Client/Multi/Columns/MultiStreamColumns_AssignProperties.cpp
--- a/Engine/Plugins/Developer/Concert/ConcertApp/MultiUserClient/Source/MultiUserClient/Private/Widgets/ActiveSession/Replication/Client/Multi/Columns/MultiStreamColumns_AssignProperties.cpp
+++ b/Engine/Plugins/Developer/Concert/ConcertApp/MultiUserClient/Source/MultiUserClient/Private/Widgets/ActiveSession/Replication/Client/Multi/Columns/MultiStreamColumns_AssignProperties.cpp
@@ -25,15 +25,17 @@ namespace UE::MultiUserClient::MultiStreamColumns
 	{
 		static void ForEachStreamAssignedTo(
 			const ConcertSharedSlate::IMultiReplicationStreamEditor& MultiEditor,
-			const FConcertPropertyChain& Property,
+			const ConcertSharedSlate::FPropertyTreeRowContext& InItem,
 			TFunctionRef<void(const TSharedRef<ConcertSharedSlate::IEditableReplicationStreamModel>& Stream)> Consume
 		)
 		{
+			const FConcertPropertyChain& Property = InItem.RowData.GetProperty();
+		
 			for (const TSharedRef<ConcertSharedSlate::IEditableReplicationStreamModel>& Stream : MultiEditor.GetMultiStreamModel().GetEditableStreams())
 			{
-				for (const TSoftObjectPtr<>& SelectedObject : MultiEditor.GetEditorBase().GetSelectedObjects())
+				for (const TSoftObjectPtr<>& ContextObject : InItem.RowData.GetContextObjects())
 				{
-					const bool bStreamAssignedToProperty = Stream->HasProperty(SelectedObject.GetUniqueID(), Property);
+					const bool bStreamAssignedToProperty = Stream->HasProperty(ContextObject.GetUniqueID(), Property);
 					if (bStreamAssignedToProperty)
 					{
 						Consume(Stream);
@@ -122,7 +124,9 @@ namespace UE::MultiUserClient::MultiStreamColumns
 			
 			virtual void PopulateSearchString(const FPropertyTreeRowContext& InItem, TArray<FString>& InOutSearchStrings) const override
 			{
-				AssignPropertyColumnUtils::ForEachStreamAssignedTo(*MultiStreamEditor.Get(), InItem.RowData.GetProperty(),
+				AssignPropertyColumnUtils::ForEachStreamAssignedTo(
+					*MultiStreamEditor.Get(),
+					InItem,
 					[this, &InOutSearchStrings](const TSharedRef<IEditableReplicationStreamModel>& Stream)
 					{
 						InOutSearchStrings.Add(AssignPropertyColumnUtils::GetClientDisplayText(*ConcertClient, ClientManager, *Stream));
diff --git a/Engine/Plugins/Developer/Concert/ConcertUI/ConcertSharedSlate/Source/ConcertSharedSlate/Private/Replication/Editor/View/Extension/CategoryRowGeneration.cpp b/Engine/Plugins/Developer/Concert/ConcertUI/ConcertSharedSlate/Source/ConcertSharedSlate/Private/Replication/Editor/View/Extension/CategoryRowGeneration.cpp
--- a/Engine/Plugins/Developer/Concert/ConcertUI/ConcertSharedSlate/Source/ConcertSharedSlate/Private/Replication/Editor/View/Extension/CategoryRowGeneration.cpp
+++ b/Engine/Plugins/Developer/Concert/ConcertUI/ConcertSharedSlate/Source/ConcertSharedSlate/Private/Replication/Editor/View/Extension/CategoryRowGeneration.cpp
@@ -19,7 +19,6 @@ namespace UE::ConcertSharedSlate
 
 			SLATE_BEGIN_ARGS(SDefaultCategoryRow){}
 				SLATE_ARGUMENT(FText, Label)
-				SLATE_ATTRIBUTE(FText, HighlightText)
 			SLATE_END_ARGS()
 
 			void Construct(const FArguments& InArgs)
@@ -27,9 +26,9 @@ namespace UE::ConcertSharedSlate
 				Label = InArgs._Label;
 				ChildSlot
 				[
+					// Do not give this widget any highlight text because category rows are not part of the search
 					SNew(STextBlock)
 					.Text(Label)
-					.HighlightText(InArgs._HighlightText)
 					.Font(FAppStyle::Get().GetFontStyle("DetailsView.CategoryFontStyle"))
 					.TextStyle(FAppStyle::Get(), "DetailsView.CategoryTextStyle")
 				];
@@ -58,10 +57,10 @@ namespace UE::ConcertSharedSlate
 			const FText Label = ContextObjects.IsEmpty()
 				? FText::GetEmpty()
 				: NameModel->GetObjectDisplayName(ContextObjects[0]);
-			
+
+			// Do not give this widget any highlight text because category rows are not part of the search
 			return SNew(Private::SDefaultCategoryRow)
-				.Label(Label)
-				.HighlightText(Args.HighlightText);
+				.Label(Label);
 		});
 	}
 }
diff --git a/Engine/Plugins/Developer/Concert/ConcertUI/ConcertSharedSlate/Source/ConcertSharedSlate/Private/Replication/Editor/View/Property/SCategoryColumnRow.cpp b/Engine/Plugins/Developer/Concert/ConcertUI/ConcertSharedSlate/Source/ConcertSharedSlate/Private/Replication/Editor/View/Property/SCategoryColumnRow.cpp
--- a/Engine/Plugins/Developer/Concert/ConcertUI/ConcertSharedSlate/Source/ConcertSharedSlate/Private/Replication/Editor/View/Property/SCategoryColumnRow.cpp
+++ b/Engine/Plugins/Developer/Concert/ConcertUI/ConcertSharedSlate/Source/ConcertSharedSlate/Private/Replication/Editor/View/Property/SCategoryColumnRow.cpp
@@ -16,8 +16,6 @@ namespace UE::ConcertSharedSlate
 {
 	void SCategoryColumnRow::Construct(const FArguments& InArgs, TSharedRef<STableViewBase> InOwner)
 	{
-		const FLinearColor BackgroundTint(0.6f, 0.6f, 0.6f, 1.0f);
-
 		// rebuilds the whole table row from scratch
 		ChildSlot
 		.Padding(0.0f, 2.0f, 0.0f, 0.0f)
diff --git a/Engine/Plugins/Developer/Concert/ConcertUI/ConcertSharedSlate/Source/ConcertSharedSlate/Private/Replication/Editor/View/Tree/SReplicationTreeView.h b/Engine/Plugins/Developer/Concert/ConcertUI/ConcertSharedSlate/Source/ConcertSharedSlate/Private/Replication/Editor/View/Tree/SReplicationTreeView.h
--- a/Engine/Plugins/Developer/Concert/ConcertUI/ConcertSharedSlate/Source/ConcertSharedSlate/Private/Replication/Editor/View/Tree/SReplicationTreeView.h
+++ b/Engine/Plugins/Developer/Concert/ConcertUI/ConcertSharedSlate/Source/ConcertSharedSlate/Private/Replication/Editor/View/Tree/SReplicationTreeView.h
@@ -349,10 +349,11 @@ namespace UE::ConcertSharedSlate
 		 */
 		TMap<FName, TSharedRef<IReplicationTreeColumn<TItemType>>> ColumnInstances;
 
+		/** Item source passed to the STreeView. */
 		TArray<TSharedPtr<TItemType>>* AllRootItems = nullptr;
-		/** Contains only the root items that passed the filters */
+		/** Contains only the root items that passed the filters. */
 		TArray<TSharedPtr<TItemType>> FilteredRootItems;
-		/** Includes ALL items in the hierarchy that have passed the filter. */
+		/** Includes ALL items in the hierarchy that have itself passed the filter or one of its children has passed the filter. */
 		TSet<TSharedPtr<TItemType>> AllFilteredItems;
 
 		struct FItemMetaData
@@ -427,6 +428,9 @@ namespace UE::ConcertSharedSlate
 		/** Callback into tree view when expansion state is changed. */
 		void OnItemExpansionChanged(TSharedPtr<TItemType> Item, bool bIsExpanded);
 
+		/** @return Whether currently searching. */
+		bool IsSearching() const { return !SearchTextFilter->GetRawFilterText().IsEmpty(); }
+
 		// Sorting
 		void Resort();
 		void Sort(TArray<TSharedPtr<TItemType>>& Items);
@@ -735,11 +739,11 @@ namespace UE::ConcertSharedSlate
 			});
 			bPassesAtLeastOnce = bChildPassedAtLeastOnce;
 		}
-		
-		if (ItemFilterResult == EItemFilterResult::Include
-			 || (ItemFilterResult == EItemFilterResult::IncludeOnlyIfChildIsIncluded && bChildPassedAtLeastOnce))
+
+		bPassesAtLeastOnce |= ItemFilterResult == EItemFilterResult::Include
+			 || (ItemFilterResult == EItemFilterResult::IncludeOnlyIfChildIsIncluded && bChildPassedAtLeastOnce);
+		if (bPassesAtLeastOnce)
 		{
-			bPassesAtLeastOnce = true;
 			FilteredItemsToShow.Add(Item);
 		}
 
@@ -761,8 +765,15 @@ namespace UE::ConcertSharedSlate
 	template <typename TItemType>
 	void SReplicationTreeView<TItemType>::CleanseItemMetaData()
 	{
+		// While searching, all items are force expanded.
+		// Do not remove items from ItemMetaData while searching because we want to restore the expansion states after search is done.
+		if (IsSearching())
+		{
+			return;
+		}
+		
 		TMap<TSharedPtr<TItemType>, FItemMetaData> NewItemMetaData;
-		for (const TSharedPtr<TItemType>& Item : *AllRootItems)
+		for (const TSharedPtr<TItemType>& Item : AllFilteredItems)
 		{
 			NewItemMetaData.Add(Item, ItemMetaData.FindOrAdd(Item));
 		}
@@ -773,12 +784,18 @@ namespace UE::ConcertSharedSlate
 	void SReplicationTreeView<TItemType>::ReapplyExpansionStates()
 	{
 		// While searching, expand all items
-		bForceParentItemsExpanded = !SearchTextFilter->GetRawFilterText().IsEmpty();
-		
-		for (const TPair<TSharedPtr<TItemType>, FItemMetaData> MetaDataPair : ItemMetaData)
+		bForceParentItemsExpanded = IsSearching();
+
+		for (const TSharedPtr<TItemType> Item : AllFilteredItems)
 		{
-			const TSharedPtr<TItemType>& Item = MetaDataPair.Key;
-			TreeView->SetItemExpansion(Item, MetaDataPair.Value.bIsExpanded || bForceParentItemsExpanded);
+			if (const FItemMetaData* ItemInfo = ItemMetaData.Find(Item))
+			{
+				TreeView->SetItemExpansion(Item, ItemInfo->bIsExpanded || bForceParentItemsExpanded);
+			}
+			else
+			{
+				TreeView->SetItemExpansion(Item, bForceParentItemsExpanded);
+			}
 		}
 	}
 
