Remove view option for filtering subobjects

diff --git a/Engine/Plugins/Developer/Concert/ConcertUI/ConcertSharedSlate/Source/ConcertSharedSlate/Private/Replication/Editor/View/ObjectViewer/SReplicationStreamViewer.cpp b/Engine/Plugins/Developer/Concert/ConcertUI/ConcertSharedSlate/Source/ConcertSharedSlate/Private/Replication/Editor/View/ObjectViewer/SReplicationStreamViewer.cpp
--- a/Engine/Plugins/Developer/Concert/ConcertUI/ConcertSharedSlate/Source/ConcertSharedSlate/Private/Replication/Editor/View/ObjectViewer/SReplicationStreamViewer.cpp
+++ b/Engine/Plugins/Developer/Concert/ConcertUI/ConcertSharedSlate/Source/ConcertSharedSlate/Private/Replication/Editor/View/ObjectViewer/SReplicationStreamViewer.cpp
@@ -248,7 +248,26 @@ namespace UE::ConcertSharedSlate
 			? InArgs._SecondaryObjectSort
 			: FColumnSortInfo{ ReplicationColumns::TopLevel::LabelColumnId, EColumnSortMode::Ascending };
 
-		ObjectViewOptions.OnDisplaySubobjectsToggled().AddSP(this, &SReplicationStreamViewer::OnSubobjectViewOptionToggled);
+		const TSharedRef<SWidget> RightOfSearch = !ObjectHierarchy
+			? InArgs._RightOfObjectSearchBar.Widget
+			// If the API user specifies an object hierarchy, then display view options for showing the actors' subobjects.
+			: [this, &InArgs]()
+			{
+				ObjectViewOptions.OnDisplaySubobjectsToggled().AddSP(this, &SReplicationStreamViewer::OnSubobjectViewOptionToggled);
+				return SNew(SHorizontalBox)
+
+					+SHorizontalBox::Slot()
+					.AutoWidth()
+					[
+						InArgs._RightOfObjectSearchBar.Widget
+					]
+				
+					+SHorizontalBox::Slot()
+					.AutoWidth()
+					[
+						ObjectViewOptions.MakeViewOptionsComboButton()
+					];
+			}();
 		
 		return SAssignNew(ReplicatedObjects, SReplicationTreeView<FReplicatedObjectData>)
 			.RootItemsSource(&RootObjectRowData)
@@ -267,19 +286,7 @@ namespace UE::ConcertSharedSlate
 			.LeftOfSearchBar() [ InArgs._LeftOfObjectSearchBar.Widget ]
 			.RightOfSearchBar()
 			[
-				SNew(SHorizontalBox)
-
-				+SHorizontalBox::Slot()
-				.AutoWidth()
-				[
-					InArgs._RightOfObjectSearchBar.Widget
-				]
-				
-				+SHorizontalBox::Slot()
-				.AutoWidth()
-				[
-					ObjectViewOptions.MakeViewOptionsComboButton()
-				]
+				RightOfSearch
 			]
 			.NoItemsContent() [ SNew(STextBlock).Text(NoObjectsAttribute) ]
 			.GetHoveredRowContent(InArgs._GetHoveredRowContent)
