Fix: Edit button menu options would not refresh correctly for remote clients due to typo

diff --git a/Engine/Plugins/Developer/Concert/ConcertUI/ConcertSharedSlate/Source/ConcertSharedSlate/Private/Replication/Editor/View/Property/SMultiObjectAssignment.cpp b/Engine/Plugins/Developer/Concert/ConcertUI/ConcertSharedSlate/Source/ConcertSharedSlate/Private/Replication/Editor/View/Property/SMultiObjectAssignment.cpp
--- a/Engine/Plugins/Developer/Concert/ConcertUI/ConcertSharedSlate/Source/ConcertSharedSlate/Private/Replication/Editor/View/Property/SMultiObjectAssignment.cpp
+++ b/Engine/Plugins/Developer/Concert/ConcertUI/ConcertSharedSlate/Source/ConcertSharedSlate/Private/Replication/Editor/View/Property/SMultiObjectAssignment.cpp
@@ -153,7 +153,7 @@ namespace UE::ConcertSharedSlate
 		
 			TreeView->RefreshPropertyData(Entries, bCanReusePropertyData);
 
-			if (!bObjectsHaveChanged)
+			if (bObjectsHaveChanged)
 			{
 				OnObjectGroupsChangedDelegate.Broadcast();
 			}
