Remove unused enum

diff --git a/Engine/Plugins/Developer/Concert/ConcertUI/ConcertSharedSlate/Source/ConcertSharedSlate/Public/Replication/Editor/Model/IEditableReplicationStreamModel.h b/Engine/Plugins/Developer/Concert/ConcertUI/ConcertSharedSlate/Source/ConcertSharedSlate/Public/Replication/Editor/Model/IEditableReplicationStreamModel.h
--- a/Engine/Plugins/Developer/Concert/ConcertUI/ConcertSharedSlate/Source/ConcertSharedSlate/Public/Replication/Editor/Model/IEditableReplicationStreamModel.h
+++ b/Engine/Plugins/Developer/Concert/ConcertUI/ConcertSharedSlate/Source/ConcertSharedSlate/Public/Replication/Editor/Model/IEditableReplicationStreamModel.h
@@ -25,13 +25,6 @@ namespace UE::ConcertSharedSlate
 		 */
 		ExternalChange
 	};
-
-	enum class EReplicatedAddObjectFlags : uint8
-	{
-		None,
-
-		SkipStreamExtension
-	};
 	
 	/**
 	 * Abstracts the concept of mapping objects to properties. This allows writing.
