Actors are now greyed out in replication Add menu if it was added by a remote client

diff --git a/Engine/Plugins/Developer/Concert/ConcertUI/ConcertSharedSlate/Source/ConcertSharedSlate/Private/Replication/Editor/View/ObjectEditor/SBaseReplicationStreamEditor.cpp b/Engine/Plugins/Developer/Concert/ConcertUI/ConcertSharedSlate/Source/ConcertSharedSlate/Private/Replication/Editor/View/ObjectEditor/SBaseReplicationStreamEditor.cpp
--- a/Engine/Plugins/Developer/Concert/ConcertUI/ConcertSharedSlate/Source/ConcertSharedSlate/Private/Replication/Editor/View/ObjectEditor/SBaseReplicationStreamEditor.cpp
+++ b/Engine/Plugins/Developer/Concert/ConcertUI/ConcertSharedSlate/Source/ConcertSharedSlate/Private/Replication/Editor/View/ObjectEditor/SBaseReplicationStreamEditor.cpp
@@ -315,7 +315,8 @@ namespace UE::ConcertSharedSlate
 			}),
 			FBuilderDelegates::FIsItemSelected::CreateLambda([this](const FSelectableObjectInfo& Item)
 			{
-				return EditablePropertiesModel->ContainsObjects({ Item.Object.Get() } );
+				const UObject* Object = Item.Object.Get();
+				return !Object || ReplicationViewer->IsDisplayedInTopView(Object);
 			}),
 			TAttribute<bool>::CreateLambda([this]() { return !IsEditingDisabled(); }),
 			TAttribute<FText>::CreateLambda([this]() { return GetEditingDisabledText(); }),
diff --git a/Engine/Plugins/Developer/Concert/ConcertUI/ConcertSharedSlate/Source/ConcertSharedSlate/Private/Replication/Editor/View/ObjectViewer/SReplicationStreamViewer.cpp b/Engine/Plugins/Developer/Concert/ConcertUI/ConcertSharedSlate/Source/ConcertSharedSlate/Private/Replication/Editor/View/ObjectViewer/SReplicationStreamViewer.cpp
--- a/Engine/Plugins/Developer/Concert/ConcertUI/ConcertSharedSlate/Source/ConcertSharedSlate/Private/Replication/Editor/View/ObjectViewer/SReplicationStreamViewer.cpp
+++ b/Engine/Plugins/Developer/Concert/ConcertUI/ConcertSharedSlate/Source/ConcertSharedSlate/Private/Replication/Editor/View/ObjectViewer/SReplicationStreamViewer.cpp
@@ -147,7 +147,17 @@ namespace UE::ConcertSharedSlate
 			ReplicatedObjects->SetExpandedItems(ItemsToExpand, true);
 		}
 	}
-	
+
+	bool SReplicationStreamViewer::IsDisplayedInTopView(const FSoftObjectPath& Object) const
+	{
+		const bool bIsContainedActor = PropertiesModel->ContainsObjects({ Object } );
+		const bool bIsContainedSubobject = PropertiesModel->AnyOfSubobjects(Object, [this](const FSoftObjectPath& SubobjectPath)
+		{
+			return PropertiesModel->ContainsObjects({ SubobjectPath } );
+		});
+		return bIsContainedActor || bIsContainedSubobject;
+	}
+
 	TArray<TSharedPtr<FReplicatedObjectData>> SReplicationStreamViewer::GetSelectedObjectItems() const
 	{
 		TArray<TSharedPtr<FReplicatedObjectData>> SelectedItems = ReplicatedObjects->GetSelectedItems();
diff --git a/Engine/Plugins/Developer/Concert/ConcertUI/ConcertSharedSlate/Source/ConcertSharedSlate/Private/Replication/Editor/View/ObjectViewer/SReplicationStreamViewer.h b/Engine/Plugins/Developer/Concert/ConcertUI/ConcertSharedSlate/Source/ConcertSharedSlate/Private/Replication/Editor/View/ObjectViewer/SReplicationStreamViewer.h
--- a/Engine/Plugins/Developer/Concert/ConcertUI/ConcertSharedSlate/Source/ConcertSharedSlate/Private/Replication/Editor/View/ObjectViewer/SReplicationStreamViewer.h
+++ b/Engine/Plugins/Developer/Concert/ConcertUI/ConcertSharedSlate/Source/ConcertSharedSlate/Private/Replication/Editor/View/ObjectViewer/SReplicationStreamViewer.h
@@ -106,6 +106,9 @@ namespace UE::ConcertSharedSlate
 		/** Expands the given objects, recursively if desired. */
 		void ExpandObjects(TConstArrayView<TSoftObjectPtr<>> Objects, bool bRecursive, bool bAtEndOfTick = false);
 
+		/** @return Whether Object is being displayed in the top object panel. */
+		bool IsDisplayedInTopView(const FSoftObjectPath& Object) const;
+
 		/** @return Gets the root objects selected in the outliner; the subobject view chooses which of these objects (or their subobjects) end up in GetSelectedObjectShowingProperties. */
 		TArray<TSharedPtr<FReplicatedObjectData>> GetSelectedObjectItems() const;
 
diff --git a/Engine/Plugins/Developer/Concert/ConcertUI/ConcertSharedSlate/Source/ConcertSharedSlate/Public/Replication/Editor/Model/IReplicationStreamModel.h b/Engine/Plugins/Developer/Concert/ConcertUI/ConcertSharedSlate/Source/ConcertSharedSlate/Public/Replication/Editor/Model/IReplicationStreamModel.h
--- a/Engine/Plugins/Developer/Concert/ConcertUI/ConcertSharedSlate/Source/ConcertSharedSlate/Public/Replication/Editor/Model/IReplicationStreamModel.h
+++ b/Engine/Plugins/Developer/Concert/ConcertUI/ConcertSharedSlate/Source/ConcertSharedSlate/Public/Replication/Editor/Model/IReplicationStreamModel.h
@@ -159,6 +159,21 @@ namespace UE::ConcertSharedSlate
 			});
 			return Subobjects;
 		}
+
+		/** @return Whether the predicate holds true for any of Root's subobjects. */
+		template<typename TLambda> requires std::is_invocable_r_v<bool, TLambda, const FSoftObjectPath&>
+		bool AnyOfSubobjects(const FSoftObjectPath& Root, TLambda&& Lambda)
+		{
+			bool bResult = false;
+			ForEachSubobject(
+				Root,
+				[&Lambda, &bResult](const FSoftObjectPath& Child)
+				{
+					bResult = Lambda(Child);
+					return bResult ? EBreakBehavior::Break : EBreakBehavior::Continue;
+				});
+			return bResult;
+		}
 		
 		virtual ~IReplicationStreamModel() = default;
 	};
