Display replication activities in session list view

diff --git a/Engine/Plugins/Developer/Concert/ConcertApp/DisasterRecoveryClient/Source/DisasterRecoveryClient/Private/SDisasterRecoveryHub.cpp b/Engine/Plugins/Developer/Concert/ConcertApp/DisasterRecoveryClient/Source/DisasterRecoveryClient/Private/SDisasterRecoveryHub.cpp
--- a/Engine/Plugins/Developer/Concert/ConcertApp/DisasterRecoveryClient/Source/DisasterRecoveryClient/Private/SDisasterRecoveryHub.cpp
+++ b/Engine/Plugins/Developer/Concert/ConcertApp/DisasterRecoveryClient/Source/DisasterRecoveryClient/Private/SDisasterRecoveryHub.cpp
@@ -588,6 +588,7 @@ TSharedRef<SWidget> SDisasterRecoveryHub::MakeSessionActivityView()
 		.IsLockActivityFilteringEnabled(false)       // Not valuable and ignored by Disaster Recovery (DR). DR sessions are for the local user only but lock/unlock are not recoverable but this will also ignore when the local user locks/unlocks assets in a Multi-User session.
 		.IsPackageActivityFilteringEnabled(true)     // Enabled and displayed by default.
 		.IsTransactionActivityFilteringEnabled(true) // Enabled and displayed by default.
+		.IsReplicationActivityFilteringEnabled(true) // Enabled and displayed by default.
 		.IsIgnoredActivityFilteringEnabled(true)     // Events ignored when restoring are not displayed by default. Enabled to inspect Multi-User transaction/package activities recorded (but not recoverable) by disaster recovery session in case the crash occurred during the Multi-User session.
 		.AreRecoverAllAndCancelButtonsVisible(false) // Replaced by this widget own (for better placement)
 		.IsRecoverThroughButtonsVisible(this, &SDisasterRecoveryHub::IsRecoverThroughButtonVisible)
diff --git a/Engine/Plugins/Developer/Concert/ConcertApp/MultiUserServer/Source/MultiUserServer/Private/Widgets/SessionTabs/Live/PackageViewer/SConcertSessionPackageViewer.cpp b/Engine/Plugins/Developer/Concert/ConcertApp/MultiUserServer/Source/MultiUserServer/Private/Widgets/SessionTabs/Live/PackageViewer/SConcertSessionPackageViewer.cpp
--- a/Engine/Plugins/Developer/Concert/ConcertApp/MultiUserServer/Source/MultiUserServer/Private/Widgets/SessionTabs/Live/PackageViewer/SConcertSessionPackageViewer.cpp
+++ b/Engine/Plugins/Developer/Concert/ConcertApp/MultiUserServer/Source/MultiUserServer/Private/Widgets/SessionTabs/Live/PackageViewer/SConcertSessionPackageViewer.cpp
@@ -38,6 +38,7 @@ void SConcertSessionPackageViewer::Construct(const FArguments& InArgs)
 		.LockActivitiesVisibility(EVisibility::Collapsed)
 		.PackageActivitiesVisibility(EVisibility::Visible)
 		.TransactionActivitiesVisibility(EVisibility::Collapsed)
+		.ReplicationActivitiesVisibility(EVisibility::Collapsed)
 		.DetailsAreaVisibility(EVisibility::Collapsed)
 		.IsAutoScrollEnabled(true)
 		.ColumnVisibilitySnapshot(UMultiUserServerColumnVisibilitySettings::GetSettings()->GetLiveSessionContentColumnVisibility())
diff --git a/Engine/Plugins/Developer/Concert/ConcertUI/ConcertSharedSlate/Source/ConcertSharedSlate/Private/SConcertSessionRecovery.cpp b/Engine/Plugins/Developer/Concert/ConcertUI/ConcertSharedSlate/Source/ConcertSharedSlate/Private/SConcertSessionRecovery.cpp
--- a/Engine/Plugins/Developer/Concert/ConcertUI/ConcertSharedSlate/Source/ConcertSharedSlate/Private/SConcertSessionRecovery.cpp
+++ b/Engine/Plugins/Developer/Concert/ConcertUI/ConcertSharedSlate/Source/ConcertSharedSlate/Private/SConcertSessionRecovery.cpp
@@ -28,6 +28,7 @@ void SConcertSessionRecovery::Construct(const FArguments& InArgs)
 	ActivityViewOptions->bEnableLockActivityFiltering = InArgs._IsLockActivityFilteringEnabled;
 	ActivityViewOptions->bEnablePackageActivityFiltering = InArgs._IsPackageActivityFilteringEnabled;
 	ActivityViewOptions->bEnableTransactionActivityFiltering = InArgs._IsTransactionActivityFilteringEnabled;
+	ActivityViewOptions->bEnableReplicationActivityFiltering = InArgs._IsReplicationActivityFilteringEnabled;
 	ActivityViewOptions->bEnableIgnoredActivityFiltering = InArgs._IsIgnoredActivityFilteringEnabled;
 	
 	SAssignNew(ActivityView, SConcertSessionActivities)
@@ -41,6 +42,7 @@ void SConcertSessionRecovery::Construct(const FArguments& InArgs)
 		.LockActivitiesVisibility(ActivityViewOptions.Get(), &FConcertSessionActivitiesOptions::GetLockActivitiesVisibility)
 		.PackageActivitiesVisibility(ActivityViewOptions.Get(), &FConcertSessionActivitiesOptions::GetPackageActivitiesVisibility)
 		.TransactionActivitiesVisibility(ActivityViewOptions.Get(), &FConcertSessionActivitiesOptions::GetTransactionActivitiesVisibility)
+		.ReplicationActivitiesVisibility(ActivityViewOptions.Get(), &FConcertSessionActivitiesOptions::GetReplicationActivitiesVisibility)
 		.IgnoredActivitiesVisibility(ActivityViewOptions.Get(), &FConcertSessionActivitiesOptions::GetIgnoredActivitiesVisibility)
 		.DetailsAreaVisibility(InArgs._DetailsAreaVisibility)
 		.NoActivitiesReasonText(InArgs._NoActivitiesReasonText);
diff --git a/Engine/Plugins/Developer/Concert/ConcertUI/ConcertSharedSlate/Source/ConcertSharedSlate/Private/Session/Activity/SConcertSessionActivities.cpp b/Engine/Plugins/Developer/Concert/ConcertUI/ConcertSharedSlate/Source/ConcertSharedSlate/Private/Session/Activity/SConcertSessionActivities.cpp
--- a/Engine/Plugins/Developer/Concert/ConcertUI/ConcertSharedSlate/Source/ConcertSharedSlate/Private/Session/Activity/SConcertSessionActivities.cpp
+++ b/Engine/Plugins/Developer/Concert/ConcertUI/ConcertSharedSlate/Source/ConcertSharedSlate/Private/Session/Activity/SConcertSessionActivities.cpp
@@ -14,6 +14,7 @@
 #include "Misc/AsyncTaskNotification.h"
 #include "Framework/Application/SlateApplication.h"
 #include "Framework/MultiBox/MultiBoxBuilder.h"
+#include "Replication/Messages/ReplicationActivity.h"
 #include "Session/Activity/ActivityColumn.h"
 #include "Widgets/SUndoHistoryDetails.h"
 #include "Widgets/Images/SThrobber.h"
@@ -31,6 +32,7 @@ const FName ShowConnectionActivitiesCheckBoxId  = TEXT("ShowConnectionActivities
 const FName ShowLockActivitiesCheckBoxId        = TEXT("ShowLockActivities");
 const FName ShowPackageActivitiesCheckBoxId     = TEXT("ShowPackageActivities");
 const FName ShowTransactionActivitiesCheckBoxId = TEXT("ShowTransactionActivities");
+const FName ShowReplicationActivitiesCheckBoxId = TEXT("ShowReplicationActivities");
 const FName ShowIgnoredActivitiesCheckBoxId     = TEXT("ShowIgnoredActivities");
 	
 FText GetSummary(const FConcertSessionActivity& Activity, const FText& ClientName, bool bAsRichText)
@@ -219,6 +221,7 @@ void SConcertSessionActivities::Construct(const FArguments& InArgs)
 	LockActivitiesVisibility = InArgs._LockActivitiesVisibility;
 	PackageActivitiesVisibility = InArgs._PackageActivitiesVisibility;
 	TransactionActivitiesVisibility = InArgs._TransactionActivitiesVisibility;
+	ReplicationActivitiesVisibility = InArgs._ReplicationActivitiesVisibility;
 	IgnoredActivitiesVisibility = InArgs._IgnoredActivitiesVisibility;
 	DetailsAreaVisibility = InArgs._DetailsAreaVisibility;
 	bAutoScrollDesired = InArgs._IsAutoScrollEnabled;
@@ -442,7 +445,7 @@ void SConcertSessionActivities::Tick(const FGeometry& AllottedGeometry, const do
 		OnActivityFilterUpdated();
 	}
 
-	FetchActivities(); // Check if we should fetch more activities in case we filtered out to many of them.
+	FetchActivities(); // Check if we should fetch more activities in case we filtered out too many of them.
 }
 
 FText SConcertSessionActivities::GetNoDetailsText() const
@@ -506,32 +509,46 @@ void SConcertSessionActivities::UpdateDetailArea(TSharedPtr<FConcertSessionActiv
 	}
 	else if (InSelectedActivity->EventPayload) // The event payload is already bundled in the activity stream?
 	{
-		if (InSelectedActivity->Activity.EventType == EConcertSyncActivityEventType::Transaction)
-		{
-			FConcertSyncTransactionEvent TransactionEvent;
-			InSelectedActivity->EventPayload->GetTypedPayload(TransactionEvent);
-			if (TransactionEvent.Transaction.ExportedObjects.Num())
-			{
-				DisplayTransactionDetails(*InSelectedActivity, TransactionEvent.Transaction);
-			}
-			else
-			{
-				SetDetailsPanelVisibility(NoDetailsPanel.Get());
-			}
-		}
-		else if (InSelectedActivity->Activity.EventType == EConcertSyncActivityEventType::Package)
+		UpdateDetailsByEventPayload(InSelectedActivity);
+	}
+	else
+	{
+		UpdateDetailsByRequest(InSelectedActivity);
+	}
+}
+
+void SConcertSessionActivities::UpdateDetailsByEventPayload(TSharedPtr<FConcertSessionActivity> InSelectedActivity)
+{
+	static_assert(static_cast<uint8>(EConcertSyncActivityEventType::Count) == 6, "If you added an EConcertSyncActivityEventType entry, you may want to update this function");
+	if (InSelectedActivity->Activity.EventType == EConcertSyncActivityEventType::Transaction)
+	{
+		FConcertSyncTransactionEvent TransactionEvent;
+		InSelectedActivity->EventPayload->GetTypedPayload(TransactionEvent);
+		if (TransactionEvent.Transaction.ExportedObjects.Num())
 		{
-			FConcertSyncPackageEvent PackageEvent;
-			InSelectedActivity->EventPayload->GetTypedPayload(PackageEvent);
-			checkf(!PackageEvent.Package.HasPackageData(), TEXT("UI should only request the package meta data because the package data is not useful and can be very large"));
-			DisplayPackageDetails(*InSelectedActivity, PackageEvent.PackageRevision, PackageEvent.Package.Info);
+			DisplayTransactionDetails(*InSelectedActivity, TransactionEvent.Transaction);
 		}
-		else // Other activity types (lock/connection) don't have details panel.
+		else
 		{
 			SetDetailsPanelVisibility(NoDetailsPanel.Get());
 		}
 	}
-	else if (InSelectedActivity->Activity.EventType == EConcertSyncActivityEventType::Transaction && GetTransactionEventFn.IsBound()) // A function is bound to get the transaction event?
+	else if (InSelectedActivity->Activity.EventType == EConcertSyncActivityEventType::Package)
+	{
+		FConcertSyncPackageEvent PackageEvent;
+		InSelectedActivity->EventPayload->GetTypedPayload(PackageEvent);
+		checkf(!PackageEvent.Package.HasPackageData(), TEXT("UI should only request the package meta data because the package data is not useful and can be very large"));
+		DisplayPackageDetails(*InSelectedActivity, PackageEvent.PackageRevision, PackageEvent.Package.Info);
+	}
+	else // Other activity types (lock/connection) don't have details panel.
+	{
+		SetDetailsPanelVisibility(NoDetailsPanel.Get());
+	}
+}
+
+void SConcertSessionActivities::UpdateDetailsByRequest(TSharedPtr<FConcertSessionActivity> InSelectedActivity)
+{
+	if (InSelectedActivity->Activity.EventType == EConcertSyncActivityEventType::Transaction && GetTransactionEventFn.IsBound()) // A function is bound to get the transaction event?
 	{
 		SetDetailsPanelVisibility(LoadingDetailsPanel.Get());
 		TWeakPtr<SConcertSessionActivities> WeakSelf = SharedThis(this);
@@ -568,6 +585,11 @@ void SConcertSessionActivities::UpdateDetailArea(TSharedPtr<FConcertSessionActiv
 			SetDetailsPanelVisibility(NoDetailsPanel.Get());
 		}
 	}
+	else if (InSelectedActivity->Activity.EventType == EConcertSyncActivityEventType::Replication)
+	{
+		// TODO UE-218508:
+		SetDetailsPanelVisibility(NoDetailsPanel.Get());
+	}
 	else
 	{
 		SetDetailsPanelVisibility(NoDetailsPanel.Get());
@@ -595,6 +617,10 @@ EConcertActivityFilterFlags SConcertSessionActivities::QueryActiveActivityFilter
 	{
 		ActiveFlags |= EConcertActivityFilterFlags::HideTransactionActivities;
 	}
+	if (ReplicationActivitiesVisibility.Get() != EVisibility::Visible)
+	{
+		ActiveFlags |= EConcertActivityFilterFlags::HideReplicationActivities;
+	}
 	if (IgnoredActivitiesVisibility.Get() != EVisibility::Visible)
 	{
 		ActiveFlags |= EConcertActivityFilterFlags::HideIgnoredActivities;
@@ -753,6 +779,10 @@ bool SConcertSessionActivities::PassesFilters(const FConcertSessionActivity& Act
 	{
 		return false;
 	}
+	else if (Activity.Activity.EventType == EConcertSyncActivityEventType::Replication && ReplicationActivitiesVisibility.Get() != EVisibility::Visible) // Filter out 'replication' activities?
+	{
+		return false;
+	}
 	else if (Activity.Activity.bIgnored && IgnoredActivitiesVisibility.Get() != EVisibility::Visible) // Filter out 'ignored' activities?
 	{
 		return false;
@@ -959,6 +989,21 @@ TSharedRef<SWidget> FConcertSessionActivitiesOptions::MakeMenuWidget(TOptional<F
 			EUserInterfaceActionType::ToggleButton
 		);
 	}
+	
+	if (bEnableReplicationActivityFiltering)
+	{
+		MenuBuilder.AddMenuEntry(
+			LOCTEXT("ShowReplicationActivities", "Show Replication Activities"),
+			LOCTEXT("ShowReplicationActivities_Tooltip", "Displays replication events"),
+			FSlateIcon(),
+			FUIAction(
+				FExecuteAction::CreateSP(this, &FConcertSessionActivitiesOptions::OnOptionToggled, ConcertSessionActivityUtils::ShowReplicationActivitiesCheckBoxId),
+				FCanExecuteAction::CreateLambda([] { return true; }),
+				FIsActionChecked::CreateLambda([this] { return bDisplayReplicationActivities; })),
+			NAME_None,
+			EUserInterfaceActionType::ToggleButton
+		);
+	}
 
 	if (bEnableIgnoredActivityFiltering)
 	{
@@ -1005,6 +1050,10 @@ void FConcertSessionActivitiesOptions::OnOptionToggled(const FName CheckBoxId)
 	{
 		bDisplayTransactionActivities = !bDisplayTransactionActivities;
 	}
+	else if (CheckBoxId == ConcertSessionActivityUtils::ShowReplicationActivitiesCheckBoxId)
+	{
+		bDisplayReplicationActivities = !bDisplayReplicationActivities;
+	}
 	else if (CheckBoxId == ConcertSessionActivityUtils::ShowIgnoredActivitiesCheckBoxId)
 	{
 		bDisplayIgnoredActivities = !bDisplayIgnoredActivities;
diff --git a/Engine/Plugins/Developer/Concert/ConcertUI/ConcertSharedSlate/Source/ConcertSharedSlate/Private/Session/Activity/SessionActivityUtils.cpp b/Engine/Plugins/Developer/Concert/ConcertUI/ConcertSharedSlate/Source/ConcertSharedSlate/Private/Session/Activity/SessionActivityUtils.cpp
--- a/Engine/Plugins/Developer/Concert/ConcertUI/ConcertSharedSlate/Source/ConcertSharedSlate/Private/Session/Activity/SessionActivityUtils.cpp
+++ b/Engine/Plugins/Developer/Concert/ConcertUI/ConcertSharedSlate/Source/ConcertSharedSlate/Private/Session/Activity/SessionActivityUtils.cpp
@@ -3,11 +3,13 @@
 #include "SessionActivityUtils.h"
 
 #include "ConcertSyncSessionTypes.h"
+#include "Replication/Messages/ReplicationActivity.h"
 
 #define LOCTEXT_NAMESPACE "SConcertSessionActivities"
 
 FText UE::ConcertSharedSlate::Private::GetOperationName(const FConcertSessionActivity& Activity)
 {
+	static_assert(static_cast<uint8>(EConcertSyncActivityEventType::Count) == 6, "If you added an EConcertSyncActivityEventType entry, you may want to update this location");
 	if (const FConcertSyncTransactionActivitySummary* TransactionSummary = Activity.ActivitySummary.Cast<FConcertSyncTransactionActivitySummary>())
 	{
 		return TransactionSummary->TransactionTitle;
@@ -55,6 +57,11 @@ FText UE::ConcertSharedSlate::Private::GetOperationName(const FConcertSessionAct
 		}
 	}
 
+	if (const FConcertSyncReplicationActivitySummary* ActivitySummary = Activity.ActivitySummary.Cast<FConcertSyncReplicationActivitySummary>())
+	{
+		return LOCTEXT("LeftReplication", "Left Replication");
+	}
+
 	return FText::GetEmpty();
 }
 
diff --git a/Engine/Plugins/Developer/Concert/ConcertUI/ConcertSharedSlate/Source/ConcertSharedSlate/Private/Session/History/SSessionHistory.cpp b/Engine/Plugins/Developer/Concert/ConcertUI/ConcertSharedSlate/Source/ConcertSharedSlate/Private/Session/History/SSessionHistory.cpp
--- a/Engine/Plugins/Developer/Concert/ConcertUI/ConcertSharedSlate/Source/ConcertSharedSlate/Private/Session/History/SSessionHistory.cpp
+++ b/Engine/Plugins/Developer/Concert/ConcertUI/ConcertSharedSlate/Source/ConcertSharedSlate/Private/Session/History/SSessionHistory.cpp
@@ -66,6 +66,7 @@ void SSessionHistory::Construct(const FArguments& InArgs)
 		.LockActivitiesVisibility(ActivityListViewOptions.Get(), &FConcertSessionActivitiesOptions::GetLockActivitiesVisibility)
 		.PackageActivitiesVisibility(ActivityListViewOptions.Get(), &FConcertSessionActivitiesOptions::GetPackageActivitiesVisibility)
 		.TransactionActivitiesVisibility(ActivityListViewOptions.Get(), &FConcertSessionActivitiesOptions::GetTransactionActivitiesVisibility)
+		.ReplicationActivitiesVisibility(ActivityListViewOptions.Get(), &FConcertSessionActivitiesOptions::GetReplicationActivitiesVisibility)
 		.DetailsAreaVisibility(InArgs._DetailsAreaVisibility)
 		.IsAutoScrollEnabled(true)
 		.ColumnVisibilitySnapshot(InArgs._ColumnVisibilitySnapshot)
diff --git a/Engine/Plugins/Developer/Concert/ConcertUI/ConcertSharedSlate/Source/ConcertSharedSlate/Public/SConcertSessionRecovery.h b/Engine/Plugins/Developer/Concert/ConcertUI/ConcertSharedSlate/Source/ConcertSharedSlate/Public/SConcertSessionRecovery.h
--- a/Engine/Plugins/Developer/Concert/ConcertUI/ConcertSharedSlate/Source/ConcertSharedSlate/Public/SConcertSessionRecovery.h
+++ b/Engine/Plugins/Developer/Concert/ConcertUI/ConcertSharedSlate/Source/ConcertSharedSlate/Public/SConcertSessionRecovery.h
@@ -28,6 +28,7 @@ public:
 		, _IsLockActivityFilteringEnabled(false)
 		, _IsPackageActivityFilteringEnabled(true)
 		, _IsTransactionActivityFilteringEnabled(true)
+		, _IsReplicationActivityFilteringEnabled(true)
 		, _IsIgnoredActivityFilteringEnabled(false)
 		, _AreRecoverAllAndCancelButtonsVisible(true)
 		, _IsRecoverThroughButtonsVisible(true) {}
@@ -76,6 +77,9 @@ public:
 
 		/** Show/hide the check box in the 'View Options' to filter transaction activities. */
 		SLATE_ARGUMENT(bool, IsTransactionActivityFilteringEnabled)
+		
+		/** Show/hide the check box in the 'View Options' to filter replication activities. */
+		SLATE_ARGUMENT(bool, IsReplicationActivityFilteringEnabled)
 
 		/** Show/hide the check box in the 'View Options' to filter unrecoverable activities. */
 		SLATE_ARGUMENT(bool, IsIgnoredActivityFilteringEnabled)
diff --git a/Engine/Plugins/Developer/Concert/ConcertUI/ConcertSharedSlate/Source/ConcertSharedSlate/Public/Session/Activity/SConcertSessionActivities.h b/Engine/Plugins/Developer/Concert/ConcertUI/ConcertSharedSlate/Source/ConcertSharedSlate/Public/Session/Activity/SConcertSessionActivities.h
--- a/Engine/Plugins/Developer/Concert/ConcertUI/ConcertSharedSlate/Source/ConcertSharedSlate/Public/Session/Activity/SConcertSessionActivities.h
+++ b/Engine/Plugins/Developer/Concert/ConcertUI/ConcertSharedSlate/Source/ConcertSharedSlate/Public/Session/Activity/SConcertSessionActivities.h
@@ -10,6 +10,7 @@
 #include "ConcertHeaderRowUtils.h"
 #include "ConcertSyncSessionTypes.h"
 #include "Misc/TextFilter.h"
+#include "Replication/Messages/ReplicationActivity.h"
 #include "Widgets/DeclarativeSyntaxSupport.h"
 #include "Widgets/SCompoundWidget.h"
 #include "Widgets/Views/SListView.h"
@@ -39,6 +40,7 @@ enum class EConcertActivityFilterFlags
 	HidePackageActivities    = 1<<2,
 	HideTransactionActivities= 1<<3,
 	HideIgnoredActivities    = 1<<4,
+	HideReplicationActivities= 1<<5
 };
 ENUM_CLASS_FLAGS(EConcertActivityFilterFlags);
 
@@ -87,6 +89,7 @@ public:
 		, _LockActivitiesVisibility(EVisibility::Hidden)
 		, _PackageActivitiesVisibility(EVisibility::Visible)
 		, _TransactionActivitiesVisibility(EVisibility::Visible)
+		, _ReplicationActivitiesVisibility(EVisibility::Hidden)
 		, _IgnoredActivitiesVisibility(EVisibility::Hidden)
 		, _DetailsAreaVisibility(EVisibility::Hidden)
 		, _IsAutoScrollEnabled(false)
@@ -131,6 +134,9 @@ public:
 
 		/** Show/hide transaction activities. */
 		SLATE_ATTRIBUTE(EVisibility, TransactionActivitiesVisibility)
+		
+		/** Show/hide replication activities. */
+		SLATE_ATTRIBUTE(EVisibility, ReplicationActivitiesVisibility)
 
 		/** Show/hide ignored activities. */
 		SLATE_ATTRIBUTE(EVisibility, IgnoredActivitiesVisibility)
@@ -237,9 +243,13 @@ private:
 	// Details
 	void OnDetailsAreaExpansionChanged(bool bExpanded);
 	void UpdateDetailArea(TSharedPtr<FConcertSessionActivity> InSelectedActivity);
+	void UpdateDetailsByEventPayload(TSharedPtr<FConcertSessionActivity> InSelectedActivity);
+	void UpdateDetailsByRequest(TSharedPtr<FConcertSessionActivity> InSelectedActivity);
+	
 	void SetDetailsPanelVisibility(const SWidget* VisiblePanel);
 	void DisplayTransactionDetails(const FConcertSessionActivity& Activity, const FConcertTransactionEventBase& InTransaction);
 	void DisplayPackageDetails(const FConcertSessionActivity& Activity, int64 PackageRevision, const FConcertPackageInfo& PackageInfo);
+	
 	FText GetNoDetailsText() const;
 	EVisibility GetDetailAreaVisibility() const { return DetailsAreaVisibility; }
 	SSplitter::ESizeRule GetDetailsAreaSizeRule() const { return bDetailsAreaExpanded ? SSplitter::ESizeRule::FractionOfParent : SSplitter::ESizeRule::SizeToContent; }
@@ -287,7 +297,10 @@ private:
 
 	/** Whether the transaction activities are displayed.*/
 	TAttribute<EVisibility> TransactionActivitiesVisibility;
-
+	
+	/** Whether the replication activities are displayed.*/
+	TAttribute<EVisibility> ReplicationActivitiesVisibility;
+	
 	/** Whether the ignored activities are displayed.*/
 	TAttribute<EVisibility> IgnoredActivitiesVisibility;
 
@@ -305,13 +318,10 @@ private:
 
 	/** Used to fetch more activities from an abstract source. Usually mutually exclusive with Append() function. May not be bound. */
 	FFetchActivitiesFunc FetchActivitiesFn;
-
 	/** Used to map an activity endpoint ID to a client. May not be bound. */
 	FGetActivityClientInfoFunc GetActivityUserFn;
-
 	/** Used to get the transaction event to display the selected transaction activity details. May not be bound. */
 	FGetTransactionEvent GetTransactionEventFn;
-
 	/** Used to get the package event to display the selected package activity details. May not be bound. */
 	FGetPackageEvent GetPackageEventFn;
 
@@ -386,6 +396,8 @@ public:
 	/** Returns whether the transaction activities are listed in SConcertSessionActivities. */
 	EVisibility GetTransactionActivitiesVisibility() const { return bDisplayTransactionActivities ? EVisibility::Visible : EVisibility::Hidden; }
 	/** Returns whether the ignored activities are listed in SConcertSessionActivities. */
+	EVisibility GetReplicationActivitiesVisibility() const { return bDisplayReplicationActivities ? EVisibility::Visible : EVisibility::Hidden; }
+	/** Returns whether the ignored activities are listed in SConcertSessionActivities. */
 	EVisibility GetIgnoredActivitiesVisibility() const { return bDisplayIgnoredActivities ? EVisibility::Visible : EVisibility::Hidden; }
 
 	/** Invoked when an options is togged from the displayed menu widget. */
@@ -399,6 +411,8 @@ public:
 	bool bEnablePackageActivityFiltering = true;
 	/** Enables the 'transaction activity' filter check box (Show Transaction Activities). */
 	bool bEnableTransactionActivityFiltering = true;
+	/** Enables the 'replication activity' filter check box (Show Replication Activities). */
+	bool bEnableReplicationActivityFiltering = true;
 	/** Enables the 'ignored activity' filter check box (Show Unrecoverable Activities). */
 	bool bEnableIgnoredActivityFiltering = false;
 	/** Controls whether the time is displayed as absolute or relative. (Display Relative Time). */
@@ -411,6 +425,8 @@ public:
 	bool bDisplayPackageActivities = true;
 	/** If transaction filtering is enabled, controls whether transaction activities are filtered out.*/
 	bool bDisplayTransactionActivities = true;
+	/** If replication filtering is enabled, controls whether replication activities are filtered out. */
+	bool bDisplayReplicationActivities = false;
 	/** If ignored activity filtering is enabled, controls whether ignored activities are filtered out.*/
 	bool bDisplayIgnoredActivities = false;
 };
