Fix RightOfSearchBar argument not being passed along to lower level UI

diff --git a/Engine/Plugins/Developer/Concert/ConcertUI/ConcertSharedSlate/Source/ConcertSharedSlate/Private/Replication/Editor/View/ObjectEditor/SBaseReplicationStreamEditor.cpp b/Engine/Plugins/Developer/Concert/ConcertUI/ConcertSharedSlate/Source/ConcertSharedSlate/Private/Replication/Editor/View/ObjectEditor/SBaseReplicationStreamEditor.cpp
--- a/Engine/Plugins/Developer/Concert/ConcertUI/ConcertSharedSlate/Source/ConcertSharedSlate/Private/Replication/Editor/View/ObjectEditor/SBaseReplicationStreamEditor.cpp
+++ b/Engine/Plugins/Developer/Concert/ConcertUI/ConcertSharedSlate/Source/ConcertSharedSlate/Private/Replication/Editor/View/ObjectEditor/SBaseReplicationStreamEditor.cpp
@@ -69,6 +69,7 @@ namespace UE::ConcertSharedSlate
 						InArgs._LeftOfObjectSearchBar.Widget
 					]
 				]
+				.RightOfObjectSearchBar() [ InArgs._RightOfObjectSearchBar.Widget ]
 				.GetHoveredRowContent(this, &SBaseReplicationStreamEditor::MakeHoveredRowContent)
 				.NoOutlinerObjects(LOCTEXT("NoObjects", "Add objects to replicate"))
 		];
