Fixed some constness with ConcertUI TSoftObjectPtr usage.

diff --git a/Engine/Plugins/Developer/Concert/ConcertUI/ConcertClientSharedSlate/Source/ConcertClientSharedSlate/Private/Replication/Editor/Model/Object/EditorObjectHierarchyModel.cpp b/Engine/Plugins/Developer/Concert/ConcertUI/ConcertClientSharedSlate/Source/ConcertClientSharedSlate/Private/Replication/Editor/Model/Object/EditorObjectHierarchyModel.cpp
--- a/Engine/Plugins/Developer/Concert/ConcertUI/ConcertClientSharedSlate/Source/ConcertClientSharedSlate/Private/Replication/Editor/Model/Object/EditorObjectHierarchyModel.cpp
+++ b/Engine/Plugins/Developer/Concert/ConcertUI/ConcertClientSharedSlate/Source/ConcertClientSharedSlate/Private/Replication/Editor/Model/Object/EditorObjectHierarchyModel.cpp
@@ -95,7 +95,7 @@ namespace UE::ConcertClientSharedSlate
 		}
 	}
 
-	TOptional<ConcertSharedSlate::IObjectHierarchyModel::FParentInfo> FEditorObjectHierarchyModel::GetParentInfo(const TSoftObjectPtr<>& ChildObject) const
+	TOptional<ConcertSharedSlate::IObjectHierarchyModel::FParentInfo> FEditorObjectHierarchyModel::GetParentInfo(const TSoftObjectPtr<const UObject>& ChildObject) const
 	{
 		using namespace ConcertSharedSlate;
 
diff --git a/Engine/Plugins/Developer/Concert/ConcertUI/ConcertClientSharedSlate/Source/ConcertClientSharedSlate/Private/Replication/Editor/Model/Object/EditorObjectHierarchyModel.h b/Engine/Plugins/Developer/Concert/ConcertUI/ConcertClientSharedSlate/Source/ConcertClientSharedSlate/Private/Replication/Editor/Model/Object/EditorObjectHierarchyModel.h
--- a/Engine/Plugins/Developer/Concert/ConcertUI/ConcertClientSharedSlate/Source/ConcertClientSharedSlate/Private/Replication/Editor/Model/Object/EditorObjectHierarchyModel.h
+++ b/Engine/Plugins/Developer/Concert/ConcertUI/ConcertClientSharedSlate/Source/ConcertClientSharedSlate/Private/Replication/Editor/Model/Object/EditorObjectHierarchyModel.h
@@ -27,7 +27,7 @@ namespace UE::ConcertClientSharedSlate
 			TFunctionRef<EBreakBehavior(const TSoftObjectPtr<>& Object, ConcertSharedSlate::EChildRelationship Relationship)> Callback,
 			ConcertSharedSlate::EChildRelationshipFlags InclusionFlags
 			) const override;
-		virtual TOptional<FParentInfo> GetParentInfo(const TSoftObjectPtr<>& ChildObject) const override;
+		virtual TOptional<FParentInfo> GetParentInfo(const TSoftObjectPtr<const UObject>& ChildObject) const override;
 		//~ End IObjectHierarchy Interface
 	};
 }
diff --git a/Engine/Plugins/Developer/Concert/ConcertUI/ConcertSharedSlate/Source/ConcertSharedSlate/Private/Replication/Editor/View/ObjectEditor/SBaseReplicationStreamEditor.cpp b/Engine/Plugins/Developer/Concert/ConcertUI/ConcertSharedSlate/Source/ConcertSharedSlate/Private/Replication/Editor/View/ObjectEditor/SBaseReplicationStreamEditor.cpp
--- a/Engine/Plugins/Developer/Concert/ConcertUI/ConcertSharedSlate/Source/ConcertSharedSlate/Private/Replication/Editor/View/ObjectEditor/SBaseReplicationStreamEditor.cpp
+++ b/Engine/Plugins/Developer/Concert/ConcertUI/ConcertSharedSlate/Source/ConcertSharedSlate/Private/Replication/Editor/View/ObjectEditor/SBaseReplicationStreamEditor.cpp
@@ -155,13 +155,13 @@ namespace UE::ConcertSharedSlate
 			// Solution: Select the highest object in the hierarchy, which is usually an actor. Users usually select an actor to add in the "Add" combo button so this also makes intuitive sense.
 			// Caveat: There may be multiple hierarchies (e.g. if multiple actors were added). This is a very seldom case though: too bad.
 			Algo::TransformIf(AddedObjects, TopLevelObjects,
-				[this, &AddedObjects](const UObject* Object)
+				[this, &AddedObjects](UObject* Object)
 				{
 					const TOptional<IObjectHierarchyModel::FParentInfo> ParentInfo = ObjectHierarchy ? ObjectHierarchy->GetParentInfo(Object) : TOptional<IObjectHierarchyModel::FParentInfo>{};
 					const bool bIsTopOfHierarchy = !ParentInfo || !AddedObjects.ContainsByPredicate([&ParentInfo](UObject* AddedObject){ return AddedObject == ParentInfo->Parent; });
 					return bIsTopOfHierarchy;
 				},
-				[](const UObject* Object){ return Object; });
+				[](UObject* Object){ return Object; });
 			ReplicationViewer->SelectObjects(TopLevelObjects);
 
 			// Expand the hierarchy for all added objects for easier editing
diff --git a/Engine/Plugins/Developer/Concert/ConcertUI/ConcertSharedSlate/Source/ConcertSharedSlate/Public/Replication/Editor/Model/Object/IObjectHierarchyModel.h b/Engine/Plugins/Developer/Concert/ConcertUI/ConcertSharedSlate/Source/ConcertSharedSlate/Public/Replication/Editor/Model/Object/IObjectHierarchyModel.h
--- a/Engine/Plugins/Developer/Concert/ConcertUI/ConcertSharedSlate/Source/ConcertSharedSlate/Public/Replication/Editor/Model/Object/IObjectHierarchyModel.h
+++ b/Engine/Plugins/Developer/Concert/ConcertUI/ConcertSharedSlate/Source/ConcertSharedSlate/Public/Replication/Editor/Model/Object/IObjectHierarchyModel.h
@@ -77,12 +77,12 @@ namespace UE::ConcertSharedSlate
 		struct FParentInfo
 		{
 			/** The parent of the child object in the hierarchy. */
-			TSoftObjectPtr<> Parent;
+			TSoftObjectPtr<const UObject> Parent;
 			/** Relationship that the child has to Parent, e.g. the child could be a component of Parent. */
 			EChildRelationship Relationship;
 		};
 		/** Gets parent info for ChildObject, if it has a child. */
-		virtual TOptional<FParentInfo> GetParentInfo(const TSoftObjectPtr<>& ChildObject) const = 0;
+		virtual TOptional<FParentInfo> GetParentInfo(const TSoftObjectPtr<const UObject>& ChildObject) const = 0;
 
 		/** Util for iterating all subobjects. */
 		void ForEachChildRecursive(
