Add CreateReadOnlyStreamModel to ConcertSharedSlate

diff --git a/Engine/Plugins/Developer/Concert/ConcertUI/ConcertSharedSlate/Source/ConcertSharedSlate/Private/Replication/Editor/Model/GenericReplicationStreamModel.cpp b/Engine/Plugins/Developer/Concert/ConcertUI/ConcertSharedSlate/Source/ConcertSharedSlate/Private/Replication/Editor/Model/GenericReplicationStreamModel.cpp
--- a/Engine/Plugins/Developer/Concert/ConcertUI/ConcertSharedSlate/Source/ConcertSharedSlate/Private/Replication/Editor/Model/GenericReplicationStreamModel.cpp
+++ b/Engine/Plugins/Developer/Concert/ConcertUI/ConcertSharedSlate/Source/ConcertSharedSlate/Private/Replication/Editor/Model/GenericReplicationStreamModel.cpp
@@ -3,6 +3,7 @@
 #include "GenericReplicationStreamModel.h"
 
 #include "ConcertLogGlobal.h"
+#include "SharedReplicationStreamModelGetters.h"
 #include "Replication/PropertyChainUtils.h"
 #include "Replication/Data/ObjectReplicationMap.h"
 #include "Replication/Editor/Model/Extension/IStreamExtender.h"
@@ -76,91 +77,32 @@ namespace UE::ConcertSharedSlate
 
 	FSoftClassPath FGenericReplicationStreamModel::GetObjectClass(const FSoftObjectPath& Object) const
 	{
-		const FConcertObjectReplicationMap* ReplicationMap = ReplicationMapAttribute.Get();
-		if (!ensure(ReplicationMap))
-		{
-			return {};
-		}
-		
-		const FConcertReplicatedObjectInfo* AssignedProperties = ReplicationMap->ReplicatedObjects.Find(Object);
-		return AssignedProperties
-			? AssignedProperties->ClassPath
-			: FSoftClassPath{};
+		return SharedStreamGetters::GetObjectClass(ReplicationMapAttribute.Get(), Object);
 	}
 
 	bool FGenericReplicationStreamModel::ContainsObjects(const TSet<FSoftObjectPath>& Objects) const
 	{
-		const FConcertObjectReplicationMap* ReplicationMap = ReplicationMapAttribute.Get();
-		return ensure(ReplicationMap)
-			&& Algo::AllOf(Objects, [this, ReplicationMap](const FSoftObjectPath& ObjectPath){ return ReplicationMap->ReplicatedObjects.Contains(ObjectPath); });
+		return SharedStreamGetters::ContainsObjects(ReplicationMapAttribute.Get(), Objects);
 	}
 
 	bool FGenericReplicationStreamModel::ContainsProperties(const FSoftObjectPath& Object, const TSet<FConcertPropertyChain>& Properties) const
 	{
-		const FConcertObjectReplicationMap* ReplicationMap = ReplicationMapAttribute.Get();
-		if (!ensure(ReplicationMap))
-		{
-			return false;
-		}
-
-		const FConcertReplicatedObjectInfo* ObjectInfo = ReplicationMap->ReplicatedObjects.Find(Object);
-		return ObjectInfo
-			&& Algo::AllOf(Properties, [ObjectInfo](const FConcertPropertyChain& Property){ return ObjectInfo->PropertySelection.ReplicatedProperties.Contains(Property); });
+		return SharedStreamGetters::ContainsProperties(ReplicationMapAttribute.Get(), Object, Properties);
 	}
 
 	bool FGenericReplicationStreamModel::ForEachReplicatedObject(TFunctionRef<EBreakBehavior(const FSoftObjectPath& Object)> Delegate) const
 	{
-		const FConcertObjectReplicationMap* ReplicationMap = ReplicationMapAttribute.Get();
-		if (!ensure(ReplicationMap))
-		{
-			return false;
-		}
-
-		for (const TPair<FSoftObjectPath, FConcertReplicatedObjectInfo>& ObjectMap: ReplicationMap->ReplicatedObjects)
-		{
-			if (Delegate(ObjectMap.Key) == EBreakBehavior::Break)
-			{
-				return true;
-			}
-		}
-
-		return !ReplicationMap->ReplicatedObjects.IsEmpty();
+		return SharedStreamGetters::ForEachReplicatedObject(ReplicationMapAttribute.Get(), Delegate);
 	}
 
 	bool FGenericReplicationStreamModel::ForEachProperty(const FSoftObjectPath& Object, TFunctionRef<EBreakBehavior(const FConcertPropertyChain& Parent)> Delegate) const
 	{
-		const FConcertObjectReplicationMap* ReplicationMap = ReplicationMapAttribute.Get();
-		if (!ensure(ReplicationMap))
-		{
-			return false;
-		}
-
-		const FConcertReplicatedObjectInfo* AssignedProperties = ReplicationMap->ReplicatedObjects.Find(Object);
-		if (!AssignedProperties)
-		{
-			return false;
-		}
-
-		for (const FConcertPropertyChain& ReplicatedPropertyInfo : AssignedProperties->PropertySelection.ReplicatedProperties)
-		{
-			if (Delegate(ReplicatedPropertyInfo) == EBreakBehavior::Break)
-			{
-				return true;
-			}
-		}
-		return !AssignedProperties->PropertySelection.ReplicatedProperties.IsEmpty();
+		return SharedStreamGetters::ForEachProperty(ReplicationMapAttribute.Get(), Object, Delegate);
 	}
 
 	uint32 FGenericReplicationStreamModel::GetNumProperties(const FSoftObjectPath& Object) const
 	{
-		const FConcertObjectReplicationMap* ReplicationMap = ReplicationMapAttribute.Get();
-		if (!ensure(ReplicationMap))
-		{
-			return 0;
-		}
-		
-		const FConcertReplicatedObjectInfo* AssignedProperties = ReplicationMap->ReplicatedObjects.Find(Object);
-		return AssignedProperties ? AssignedProperties->PropertySelection.ReplicatedProperties.Num() : 0;
+		return SharedStreamGetters::GetNumProperties(ReplicationMapAttribute.Get(), Object);
 	}
 
 	void FGenericReplicationStreamModel::AddObjects(TConstArrayView<UObject*> Objects)
diff --git a/Engine/Plugins/Developer/Concert/ConcertUI/ConcertSharedSlate/Source/ConcertSharedSlate/Private/Replication/Editor/Model/GenericReplicationStreamModel.h b/Engine/Plugins/Developer/Concert/ConcertUI/ConcertSharedSlate/Source/ConcertSharedSlate/Private/Replication/Editor/Model/GenericReplicationStreamModel.h
--- a/Engine/Plugins/Developer/Concert/ConcertUI/ConcertSharedSlate/Source/ConcertSharedSlate/Private/Replication/Editor/Model/GenericReplicationStreamModel.h
+++ b/Engine/Plugins/Developer/Concert/ConcertUI/ConcertSharedSlate/Source/ConcertSharedSlate/Private/Replication/Editor/Model/GenericReplicationStreamModel.h
@@ -3,8 +3,8 @@
 #pragma once
 
 #include "Replication/Editor/Model/IEditableReplicationStreamModel.h"
+
 #include "Misc/Attribute.h"
-#include "UObject/WeakObjectPtrTemplates.h"
 
 struct FConcertStreamObjectAutoBindingRules;
 struct FConcertObjectReplicationMap;
@@ -13,7 +13,7 @@ namespace UE::ConcertSharedSlate
 {
 	class IStreamExtender;
 
-	/** Implements logic for editing a FConcertObjectReplicationMap contained in an UObject. */
+	/** Implements logic for editing a FConcertObjectReplicationMap. */
 	class FGenericReplicationStreamModel
 		: public IEditableReplicationStreamModel
 	{
@@ -45,7 +45,7 @@ namespace UE::ConcertSharedSlate
 	private:
 
 		/** Returns the replication map that is supposed to be edited. */
-		TAttribute<FConcertObjectReplicationMap*> ReplicationMapAttribute;
+		const TAttribute<FConcertObjectReplicationMap*> ReplicationMapAttribute;
 
 		/** Adds properties and objects when an object is added. Can be null. */
 		TSharedPtr<IStreamExtender> Extender;
@@ -56,5 +56,4 @@ namespace UE::ConcertSharedSlate
 		/** Applies Extender to AddedObject whilst adding any additionally added UObjects to ObjectsAddedSoFar and avoiding adding objects already added to ObjectsAddedSoFar. */
 		void ExtendObjects(FConcertObjectReplicationMap& ReplicationMap, UObject& AddedObject, TArray<UObject*>& ObjectsAddedSoFar);
 	};
-}
-
+}
\ No newline at end of file
diff --git a/Engine/Plugins/Developer/Concert/ConcertUI/ConcertSharedSlate/Source/ConcertSharedSlate/Private/Replication/Editor/Model/ReadableReplicationStreamModel.cpp b/Engine/Plugins/Developer/Concert/ConcertUI/ConcertSharedSlate/Source/ConcertSharedSlate/Private/Replication/Editor/Model/ReadableReplicationStreamModel.cpp
new file mode 100644
--- /dev/null
+++ b/Engine/Plugins/Developer/Concert/ConcertUI/ConcertSharedSlate/Source/ConcertSharedSlate/Private/Replication/Editor/Model/ReadableReplicationStreamModel.cpp
@@ -0,0 +1,44 @@
+﻿// Copyright Epic Games, Inc. All Rights Reserved.
+
+#include "ReadableReplicationStreamModel.h"
+
+#include "SharedReplicationStreamModelGetters.h"
+
+namespace UE::ConcertSharedSlate
+{
+	FReadableReplicationStreamModel::FReadableReplicationStreamModel(TAttribute<const FConcertObjectReplicationMap*> InReplicationMapAttribute)
+		: ReplicationMapAttribute(MoveTemp(InReplicationMapAttribute))
+	{
+		check(ReplicationMapAttribute.IsBound());
+	}
+
+	FSoftClassPath FReadableReplicationStreamModel::GetObjectClass(const FSoftObjectPath& Object) const
+	{
+		return SharedStreamGetters::GetObjectClass(ReplicationMapAttribute.Get(), Object);
+	}
+
+	bool FReadableReplicationStreamModel::ContainsObjects(const TSet<FSoftObjectPath>& Objects) const
+	{
+		return SharedStreamGetters::ContainsObjects(ReplicationMapAttribute.Get(), Objects);
+	}
+
+	bool FReadableReplicationStreamModel::ContainsProperties(const FSoftObjectPath& Object, const TSet<FConcertPropertyChain>& Properties) const
+	{
+		return SharedStreamGetters::ContainsProperties(ReplicationMapAttribute.Get(), Object, Properties);
+	}
+
+	bool FReadableReplicationStreamModel::ForEachReplicatedObject(TFunctionRef<EBreakBehavior(const FSoftObjectPath& Object)> Delegate) const
+	{
+		return SharedStreamGetters::ForEachReplicatedObject(ReplicationMapAttribute.Get(), Delegate);
+	}
+
+	bool FReadableReplicationStreamModel::ForEachProperty(const FSoftObjectPath& Object, TFunctionRef<EBreakBehavior(const FConcertPropertyChain& Parent)> Delegate) const
+	{
+		return SharedStreamGetters::ForEachProperty(ReplicationMapAttribute.Get(), Object, Delegate);
+	}
+
+	uint32 FReadableReplicationStreamModel::GetNumProperties(const FSoftObjectPath& Object) const
+	{
+		return SharedStreamGetters::GetNumProperties(ReplicationMapAttribute.Get(), Object);
+	}
+}
diff --git a/Engine/Plugins/Developer/Concert/ConcertUI/ConcertSharedSlate/Source/ConcertSharedSlate/Private/Replication/Editor/Model/ReadableReplicationStreamModel.h b/Engine/Plugins/Developer/Concert/ConcertUI/ConcertSharedSlate/Source/ConcertSharedSlate/Private/Replication/Editor/Model/ReadableReplicationStreamModel.h
new file mode 100644
--- /dev/null
+++ b/Engine/Plugins/Developer/Concert/ConcertUI/ConcertSharedSlate/Source/ConcertSharedSlate/Private/Replication/Editor/Model/ReadableReplicationStreamModel.h
@@ -0,0 +1,38 @@
+﻿// Copyright Epic Games, Inc. All Rights Reserved.
+
+#pragma once
+
+#include "Replication/Editor/Model/IReplicationStreamModel.h"
+
+#include "Misc/Attribute.h"
+
+struct FConcertStreamObjectAutoBindingRules;
+struct FConcertObjectReplicationMap;
+
+namespace UE::ConcertSharedSlate
+{
+	class IStreamExtender;
+
+	/** Implements logic for reading a FConcertObjectReplicationMap. */
+	class FReadableReplicationStreamModel
+		: public IReplicationStreamModel
+	{
+	public:
+		
+		FReadableReplicationStreamModel(TAttribute<const FConcertObjectReplicationMap*> InReplicationMapAttribute);
+		
+		//~ Begin IReplicationStreamModel Interface
+		virtual FSoftClassPath GetObjectClass(const FSoftObjectPath& Object) const override;
+		virtual bool ContainsObjects(const TSet<FSoftObjectPath>& Objects) const override;
+		virtual bool ContainsProperties(const FSoftObjectPath& Object, const TSet<FConcertPropertyChain>& Properties) const override;
+		virtual bool ForEachReplicatedObject(TFunctionRef<EBreakBehavior(const FSoftObjectPath& Object)> Delegate) const override;
+		virtual bool ForEachProperty(const FSoftObjectPath& Object, TFunctionRef<EBreakBehavior(const FConcertPropertyChain& Parent)> Delegate) const override;
+		virtual uint32 GetNumProperties(const FSoftObjectPath& Object) const override;
+		//~ End IReplicationStreamModel Interface
+
+	private:
+
+		/** Returns the replication map that is supposed to be edited. */
+		const TAttribute<const FConcertObjectReplicationMap*> ReplicationMapAttribute;
+	};
+}
\ No newline at end of file
diff --git a/Engine/Plugins/Developer/Concert/ConcertUI/ConcertSharedSlate/Source/ConcertSharedSlate/Private/Replication/Editor/Model/SharedReplicationStreamModelGetters.cpp b/Engine/Plugins/Developer/Concert/ConcertUI/ConcertSharedSlate/Source/ConcertSharedSlate/Private/Replication/Editor/Model/SharedReplicationStreamModelGetters.cpp
new file mode 100644
--- /dev/null
+++ b/Engine/Plugins/Developer/Concert/ConcertUI/ConcertSharedSlate/Source/ConcertSharedSlate/Private/Replication/Editor/Model/SharedReplicationStreamModelGetters.cpp
@@ -0,0 +1,93 @@
+﻿// Copyright Epic Games, Inc. All Rights Reserved.
+
+#include "SharedReplicationStreamModelGetters.h"
+
+#include "Replication/Data/ObjectReplicationMap.h"
+
+#include "Algo/AllOf.h"
+
+namespace UE::ConcertSharedSlate::SharedStreamGetters
+{
+	FSoftClassPath GetObjectClass(const FConcertObjectReplicationMap* ReplicationMap, const FSoftObjectPath& Object)
+	{
+		if (!ensure(ReplicationMap))
+		{
+			return {};
+		}
+		
+		const FConcertReplicatedObjectInfo* AssignedProperties = ReplicationMap->ReplicatedObjects.Find(Object);
+		return AssignedProperties
+			? AssignedProperties->ClassPath
+			: FSoftClassPath{};
+	}
+
+	bool ContainsObjects(const FConcertObjectReplicationMap* ReplicationMap, const TSet<FSoftObjectPath>& Objects)
+	{
+		return ensure(ReplicationMap)
+			&& Algo::AllOf(Objects, [ReplicationMap](const FSoftObjectPath& ObjectPath){ return ReplicationMap->ReplicatedObjects.Contains(ObjectPath); });
+	}
+
+	bool ContainsProperties(const FConcertObjectReplicationMap* ReplicationMap, const FSoftObjectPath& Object, const TSet<FConcertPropertyChain>& Properties)
+	{
+		if (!ensure(ReplicationMap))
+		{
+			return false;
+		}
+
+		const FConcertReplicatedObjectInfo* ObjectInfo = ReplicationMap->ReplicatedObjects.Find(Object);
+		return ObjectInfo
+			&& Algo::AllOf(Properties, [ObjectInfo](const FConcertPropertyChain& Property){ return ObjectInfo->PropertySelection.ReplicatedProperties.Contains(Property); });
+	}
+
+	bool ForEachReplicatedObject(const FConcertObjectReplicationMap* ReplicationMap, TFunctionRef<EBreakBehavior(const FSoftObjectPath& Object)> Delegate)
+	{
+		if (!ensure(ReplicationMap))
+		{
+			return false;
+		}
+
+		for (const TPair<FSoftObjectPath, FConcertReplicatedObjectInfo>& ObjectMap: ReplicationMap->ReplicatedObjects)
+		{
+			if (Delegate(ObjectMap.Key) == EBreakBehavior::Break)
+			{
+				return true;
+			}
+		}
+
+		return !ReplicationMap->ReplicatedObjects.IsEmpty();
+	}
+
+	bool ForEachProperty(const FConcertObjectReplicationMap* ReplicationMap, const FSoftObjectPath& Object, TFunctionRef<EBreakBehavior(const FConcertPropertyChain& Parent)> Delegate)
+	{
+		if (!ensure(ReplicationMap))
+		{
+			return false;
+		}
+
+		const FConcertReplicatedObjectInfo* AssignedProperties = ReplicationMap->ReplicatedObjects.Find(Object);
+		if (!AssignedProperties)
+		{
+			return false;
+		}
+
+		for (const FConcertPropertyChain& ReplicatedPropertyInfo : AssignedProperties->PropertySelection.ReplicatedProperties)
+		{
+			if (Delegate(ReplicatedPropertyInfo) == EBreakBehavior::Break)
+			{
+				return true;
+			}
+		}
+		return !AssignedProperties->PropertySelection.ReplicatedProperties.IsEmpty();
+	}
+
+	uint32 GetNumProperties(const FConcertObjectReplicationMap* ReplicationMap, const FSoftObjectPath& Object)
+	{
+		if (!ensure(ReplicationMap))
+		{
+			return 0;
+		}
+		
+		const FConcertReplicatedObjectInfo* AssignedProperties = ReplicationMap->ReplicatedObjects.Find(Object);
+		return AssignedProperties ? AssignedProperties->PropertySelection.ReplicatedProperties.Num() : 0;
+	}
+}
diff --git a/Engine/Plugins/Developer/Concert/ConcertUI/ConcertSharedSlate/Source/ConcertSharedSlate/Private/Replication/Editor/Model/SharedReplicationStreamModelGetters.h b/Engine/Plugins/Developer/Concert/ConcertUI/ConcertSharedSlate/Source/ConcertSharedSlate/Private/Replication/Editor/Model/SharedReplicationStreamModelGetters.h
new file mode 100644
--- /dev/null
+++ b/Engine/Plugins/Developer/Concert/ConcertUI/ConcertSharedSlate/Source/ConcertSharedSlate/Private/Replication/Editor/Model/SharedReplicationStreamModelGetters.h
@@ -0,0 +1,31 @@
+﻿// Copyright Epic Games, Inc. All Rights Reserved.
+
+#pragma once
+
+#include "Containers/Set.h"
+
+#include "HAL/Platform.h"
+#include "Templates/FunctionFwd.h"
+
+struct FConcertPropertyChain;
+struct FConcertObjectReplicationMap;
+struct FSoftClassPath;
+struct FSoftObjectPath;
+
+enum class EBreakBehavior : uint8;
+
+namespace UE::ConcertSharedSlate::SharedStreamGetters
+{
+	/** @return Gets the object class from ReplicationMap for Object. */
+	FSoftClassPath GetObjectClass(const FConcertObjectReplicationMap* ReplicationMap, const FSoftObjectPath& Object);
+	/** @return Whether all the objects are contained in ReplicationMap. */
+	bool ContainsObjects(const FConcertObjectReplicationMap* ReplicationMap, const TSet<FSoftObjectPath>& Objects);
+	/** @return Whether Object has the given properties assigned. */
+	bool ContainsProperties(const FConcertObjectReplicationMap* ReplicationMap, const FSoftObjectPath& Object, const TSet<FConcertPropertyChain>& Properties);
+	/** @return Whether Delegate was invokes at least once. */
+	bool ForEachReplicatedObject(const FConcertObjectReplicationMap* ReplicationMap, TFunctionRef<EBreakBehavior(const FSoftObjectPath& Object)> Delegate);
+	/** @return Whether Delegate was invokes at least once. */
+	bool ForEachProperty(const FConcertObjectReplicationMap* ReplicationMap, const FSoftObjectPath& Object, TFunctionRef<EBreakBehavior(const FConcertPropertyChain& Parent)> Delegate);
+	/** @return Number of properties assigned to Object */
+	uint32 GetNumProperties(const FConcertObjectReplicationMap* ReplicationMap, const FSoftObjectPath& Object);
+}
diff --git a/Engine/Plugins/Developer/Concert/ConcertUI/ConcertSharedSlate/Source/ConcertSharedSlate/Private/Replication/ReplicationWidgetFactories.cpp b/Engine/Plugins/Developer/Concert/ConcertUI/ConcertSharedSlate/Source/ConcertSharedSlate/Private/Replication/ReplicationWidgetFactories.cpp
--- a/Engine/Plugins/Developer/Concert/ConcertUI/ConcertSharedSlate/Source/ConcertSharedSlate/Private/Replication/ReplicationWidgetFactories.cpp
+++ b/Engine/Plugins/Developer/Concert/ConcertUI/ConcertSharedSlate/Source/ConcertSharedSlate/Private/Replication/ReplicationWidgetFactories.cpp
@@ -3,6 +3,7 @@
 #include "Replication/ReplicationWidgetFactories.h"
 
 #include "Editor/Model/GenericReplicationStreamModel.h"
+#include "Editor/Model/ReadableReplicationStreamModel.h"
 #include "Editor/View/MultiEditor/SMultiReplicationStreamEditor.h"
 #include "Editor/View/ObjectEditor/SBaseReplicationStreamEditor.h"
 #include "Editor/View/Property/SMultiObjectAssignment.h"
@@ -16,6 +17,11 @@
 
 namespace UE::ConcertSharedSlate
 {
+	TSharedRef<IReplicationStreamModel> CreateReadOnlyStreamModel(TAttribute<const FConcertObjectReplicationMap*> ReplicationMapAttribute)
+	{
+		return MakeShared<FReadableReplicationStreamModel>(MoveTemp(ReplicationMapAttribute));
+	}
+
 	TSharedRef<IEditableReplicationStreamModel> CreateBaseStreamModel(
 		TAttribute<FConcertObjectReplicationMap*> ReplicationMapAttribute,
 		TSharedPtr<IStreamExtender> Extender
diff --git a/Engine/Plugins/Developer/Concert/ConcertUI/ConcertSharedSlate/Source/ConcertSharedSlate/Public/Replication/ReplicationWidgetFactories.h b/Engine/Plugins/Developer/Concert/ConcertUI/ConcertSharedSlate/Source/ConcertSharedSlate/Public/Replication/ReplicationWidgetFactories.h
--- a/Engine/Plugins/Developer/Concert/ConcertUI/ConcertSharedSlate/Source/ConcertSharedSlate/Public/Replication/ReplicationWidgetFactories.h
+++ b/Engine/Plugins/Developer/Concert/ConcertUI/ConcertSharedSlate/Source/ConcertSharedSlate/Public/Replication/ReplicationWidgetFactories.h
@@ -33,6 +33,11 @@ namespace UE::ConcertSharedSlate
 	class IStreamExtender;
 	class IObjectHierarchyModel;
 	class IPropertySourceProcessor;
+
+	/** Creates a stream model that will only read from ReplicationMapAttribute. */
+	CONCERTSHAREDSLATE_API TSharedRef<IReplicationStreamModel> CreateReadOnlyStreamModel(
+		TAttribute<const FConcertObjectReplicationMap*> ReplicationMapAttribute
+		);
 	
 	/**
 	 * Creates a model that can be passed to CreateEditor.
