Move soft object path utils to ConcertSyncCore.

diff --git a/Engine/Plugins/Developer/Concert/ConcertSync/ConcertSyncCore/Source/ConcertSyncCore/Private/Misc/ObjectUtils.cpp b/Engine/Plugins/Developer/Concert/ConcertSync/ConcertSyncCore/Source/ConcertSyncCore/Private/Misc/ObjectUtils.cpp
new file mode 100644
--- /dev/null
+++ b/Engine/Plugins/Developer/Concert/ConcertSync/ConcertSyncCore/Source/ConcertSyncCore/Private/Misc/ObjectUtils.cpp
@@ -0,0 +1,71 @@
+// Copyright Epic Games, Inc. All Rights Reserved.
+
+#include "Misc/ObjectUtils.h"
+
+#include "Containers/UnrealString.h"
+#include "UObject/SoftObjectPath.h"
+
+namespace UE::ConcertSyncCore
+{
+	bool IsActor(const FSoftObjectPath& SoftObjectPath)
+	{
+		// Example of an actor called floor
+		// SoftObjectPath = { AssetPath = {PackageName = "/Game/Maps/SyncBoxLevel", AssetName = "SyncBoxLevel"}, SubPathString = "PersistentLevel.Floor" } }
+		const FString& SubPathString = SoftObjectPath.GetSubPathString();
+
+		constexpr int32 PersistentLevelStringLength = 16; // "PersistentLevel." has 16 characters
+		const bool bIsWorldObject = SubPathString.Contains(TEXT("PersistentLevel."), ESearchCase::CaseSensitive);
+		if (!bIsWorldObject)
+		{
+			// Not a path to a world object
+			return {};
+		}
+
+		// Start search after the . behind PersistentLevel
+		const int32 StartSearch = PersistentLevelStringLength + 1;
+		const int32 IndexOfDotAfterActorName = SubPathString.Find(TEXT("."), ESearchCase::CaseSensitive, ESearchDir::FromStart, StartSearch);
+		return IndexOfDotAfterActorName == INDEX_NONE;
+	}
+	
+	TOptional<FSoftObjectPath> GetActorOf(const FSoftObjectPath& SoftObjectPath)
+	{
+		// Example of an actor called floor
+		// SoftObjectPath = { AssetPath = {PackageName = "/Game/Maps/SyncBoxLevel", AssetName = "SyncBoxLevel"}, SubPathString = "PersistentLevel.Floor" } }
+		const FString& SubPathString = SoftObjectPath.GetSubPathString();
+
+		constexpr int32 PersistentLevelStringLength = 16; // "PersistentLevel." has 16 characters
+		const bool bIsWorldObject = SubPathString.Contains(TEXT("PersistentLevel."), ESearchCase::CaseSensitive);
+		if (!bIsWorldObject)
+		{
+			// Not a path to a world object
+			return {};
+		}
+
+		// Start search after the . behind PersistentLevel
+		const int32 StartSearch = PersistentLevelStringLength + 1;
+		const int32 IndexOfDotAfterActorName = SubPathString.Find(TEXT("."), ESearchCase::CaseSensitive, ESearchDir::FromStart, StartSearch);
+		if (IndexOfDotAfterActorName == INDEX_NONE)
+		{
+			// SoftObjectPath points to an actor
+			return {};
+		}
+
+		const int32 NumToChopOffRight = SubPathString.Len() - IndexOfDotAfterActorName;
+		const FString NewSubstring = SubPathString.LeftChop(NumToChopOffRight);
+		const FSoftObjectPath PathToOwningActor(SoftObjectPath.GetAssetPath(), NewSubstring);
+		return PathToOwningActor;
+	}
+	
+	FString ExtractObjectNameFromPath(const FSoftObjectPath& Object)
+	{
+		// Subpath looks like this PersistentLevel.Actor.Component
+		const FString& Subpath = Object.GetSubPathString();
+		const int32 LastDotIndex = Subpath.Find(TEXT("."), ESearchCase::CaseSensitive, ESearchDir::FromEnd);
+		if (LastDotIndex == INDEX_NONE)
+		{
+			return {};
+		}
+		return Subpath.RightChop(LastDotIndex + 1);
+	}
+};
+
diff --git a/Engine/Plugins/Developer/Concert/ConcertSync/ConcertSyncCore/Source/ConcertSyncCore/Public/Misc/ObjectUtils.h b/Engine/Plugins/Developer/Concert/ConcertSync/ConcertSyncCore/Source/ConcertSyncCore/Public/Misc/ObjectUtils.h
new file mode 100644
--- /dev/null
+++ b/Engine/Plugins/Developer/Concert/ConcertSync/ConcertSyncCore/Source/ConcertSyncCore/Public/Misc/ObjectUtils.h
@@ -0,0 +1,19 @@
+// Copyright Epic Games, Inc. All Rights Reserved.
+
+#pragma once
+
+#include "Misc/Optional.h"
+
+struct FSoftObjectPath;
+
+namespace UE::ConcertSyncCore
+{
+	/** @return Checks whether the given object is an actor. */
+	CONCERTSYNCCORE_API bool IsActor(const FSoftObjectPath& Object);
+	
+	/** @return Get the owning actor of Subobject. If Subobject is an actor, then this returns unset. */
+	CONCERTSYNCCORE_API TOptional<FSoftObjectPath> GetActorOf(const FSoftObjectPath& Subobject);
+	
+	/** @return Gets the last object name in the subpath. */
+	CONCERTSYNCCORE_API FString ExtractObjectNameFromPath(const FSoftObjectPath& Object);
+};
diff --git a/Engine/Plugins/Developer/Concert/ConcertUI/ConcertClientSharedSlate/Source/ConcertClientSharedSlate/Private/Replication/ObjectNameUtils.cpp b/Engine/Plugins/Developer/Concert/ConcertUI/ConcertClientSharedSlate/Source/ConcertClientSharedSlate/Private/Replication/ObjectNameUtils.cpp
--- a/Engine/Plugins/Developer/Concert/ConcertUI/ConcertClientSharedSlate/Source/ConcertClientSharedSlate/Private/Replication/ObjectNameUtils.cpp
+++ b/Engine/Plugins/Developer/Concert/ConcertUI/ConcertClientSharedSlate/Source/ConcertClientSharedSlate/Private/Replication/ObjectNameUtils.cpp
@@ -3,7 +3,7 @@
 #include "Replication/ObjectNameUtils.h"
 
 #include "Internationalization/Text.h"
-#include "Replication/Utils/ObjectUtils.h"
+#include "Misc/ObjectUtils.h"
 
 #include "GameFramework/Actor.h"
 #include "SubobjectDataSubsystem.h"
@@ -54,6 +54,6 @@ namespace UE::ConcertClientSharedSlate
 			return FText::FromString(ResolvedObject->GetName());
 		}
 		
-		return FText::FromString(ConcertSharedSlate::ObjectUtils::ExtractObjectDisplayStringFromPath(ObjectPath.GetUniqueID()));
+		return FText::FromString(ConcertSyncCore::ExtractObjectNameFromPath(ObjectPath.GetUniqueID()));
 	}
 }
diff --git a/Engine/Plugins/Developer/Concert/ConcertUI/ConcertSharedSlate/Source/ConcertSharedSlate/Private/Replication/Editor/Utils/DisplayUtils.cpp b/Engine/Plugins/Developer/Concert/ConcertUI/ConcertSharedSlate/Source/ConcertSharedSlate/Private/Replication/Editor/Utils/DisplayUtils.cpp
--- a/Engine/Plugins/Developer/Concert/ConcertUI/ConcertSharedSlate/Source/ConcertSharedSlate/Private/Replication/Editor/Utils/DisplayUtils.cpp
+++ b/Engine/Plugins/Developer/Concert/ConcertUI/ConcertSharedSlate/Source/ConcertSharedSlate/Private/Replication/Editor/Utils/DisplayUtils.cpp
@@ -4,7 +4,7 @@
 
 #include "ConcertLogGlobal.h"
 #include "Replication/PropertyResolutionCache.h"
-#include "Replication/Utils/ObjectUtils.h"
+#include "Misc/ObjectUtils.h"
 #include "Replication/Editor/Model/IReplicationStreamModel.h"
 #include "Replication/Editor/Model/Object/IObjectNameModel.h"
 
@@ -42,7 +42,7 @@ namespace UE::ConcertSharedSlate::DisplayUtils
 	
 	FText ExtractObjectDisplayTextFromPath(const FSoftObjectPath& Object)
 	{
-		return FText::FromString(ObjectUtils::ExtractObjectDisplayStringFromPath(Object));
+		return FText::FromString(ConcertSyncCore::ExtractObjectNameFromPath(Object));
 	}
 
 	FString GetObjectDisplayString(const UObject& Object)
diff --git a/Engine/Plugins/Developer/Concert/ConcertUI/ConcertSharedSlate/Source/ConcertSharedSlate/Private/Replication/Editor/Utils/ObjectUtils.cpp b/Engine/Plugins/Developer/Concert/ConcertUI/ConcertSharedSlate/Source/ConcertSharedSlate/Private/Replication/Editor/Utils/ObjectUtils.cpp
deleted file mode 100644
diff --git a/Engine/Plugins/Developer/Concert/ConcertUI/ConcertSharedSlate/Source/ConcertSharedSlate/Private/Replication/Editor/View/ObjectEditor/SBaseReplicationStreamEditor.cpp b/Engine/Plugins/Developer/Concert/ConcertUI/ConcertSharedSlate/Source/ConcertSharedSlate/Private/Replication/Editor/View/ObjectEditor/SBaseReplicationStreamEditor.cpp
--- a/Engine/Plugins/Developer/Concert/ConcertUI/ConcertSharedSlate/Source/ConcertSharedSlate/Private/Replication/Editor/View/ObjectEditor/SBaseReplicationStreamEditor.cpp
+++ b/Engine/Plugins/Developer/Concert/ConcertUI/ConcertSharedSlate/Source/ConcertSharedSlate/Private/Replication/Editor/View/ObjectEditor/SBaseReplicationStreamEditor.cpp
@@ -12,7 +12,7 @@
 #include "Replication/Editor/View/ObjectViewer/SReplicationStreamViewer.h"
 
 #include "Algo/AnyOf.h"
-#include "Replication/Utils/ObjectUtils.h"
+#include "Misc/ObjectUtils.h"
 #include "Styling/AppStyle.h"
 #include "UObject/Class.h"
 #include "Widgets/Images/SImage.h"
@@ -105,7 +105,7 @@ namespace UE::ConcertSharedSlate
 	{
 		const FSoftObjectPath Object = Data->GetObjectPath();
 		
-		if (ObjectUtils::IsActor(Object))
+		if (ConcertSyncCore::IsActor(Object))
 		{
 			const TSharedRef<SWidget> Content =
 				SNew(SButton)
diff --git a/Engine/Plugins/Developer/Concert/ConcertUI/ConcertSharedSlate/Source/ConcertSharedSlate/Private/Replication/Editor/View/ObjectViewer/SReplicationStreamViewer.cpp b/Engine/Plugins/Developer/Concert/ConcertUI/ConcertSharedSlate/Source/ConcertSharedSlate/Private/Replication/Editor/View/ObjectViewer/SReplicationStreamViewer.cpp
--- a/Engine/Plugins/Developer/Concert/ConcertUI/ConcertSharedSlate/Source/ConcertSharedSlate/Private/Replication/Editor/View/ObjectViewer/SReplicationStreamViewer.cpp
+++ b/Engine/Plugins/Developer/Concert/ConcertUI/ConcertSharedSlate/Source/ConcertSharedSlate/Private/Replication/Editor/View/ObjectViewer/SReplicationStreamViewer.cpp
@@ -9,7 +9,7 @@
 #include "Replication/Editor/View/Column/ObjectColumnAdapter.h"
 #include "Replication/Editor/View/Column/SelectionViewerColumns.h"
 #include "Replication/Editor/View/Property/SPropertyTreeView.h"
-#include "Replication/Utils/ObjectUtils.h"
+#include "Misc/ObjectUtils.h"
 #include "SReplicatedPropertyView.h"
 #include "Trace/ConcertTrace.h"
 
@@ -167,7 +167,7 @@ namespace UE::ConcertSharedSlate
 			const FSoftObjectPath& ObjectPath = ObjectData->GetObjectPath();
 			const bool bIsInModel = PropertiesModel->ContainsObjects({ ObjectPath });
 			
-			const TOptional<FSoftObjectPath> OwningActor = ObjectUtils::GetActorOf(ObjectPath);
+			const TOptional<FSoftObjectPath> OwningActor = ConcertSyncCore::GetActorOf(ObjectPath);
 			// When displaying object from local machine...
 			// ... the "Add Actor" button has added an actor without properties to the model; however objects without assigned properties are not transmitted to server.
 			const bool bContainsOwningActor = OwningActor.IsSet() && PropertiesModel->ContainsObjects({ *OwningActor });
@@ -384,7 +384,7 @@ namespace UE::ConcertSharedSlate
 
 		PropertiesModel->ForEachReplicatedObject([this, &Delegate, &AddedActors, &PendingActors](const FSoftObjectPath& Object)
 		{
-			if (const TOptional<FSoftObjectPath> OwningActor = ObjectUtils::GetActorOf(Object)
+			if (const TOptional<FSoftObjectPath> OwningActor = ConcertSyncCore::GetActorOf(Object)
 				; OwningActor && CanDisplayObject(*OwningActor))
 			{
 				PendingActors.Add(*OwningActor);
@@ -395,7 +395,7 @@ namespace UE::ConcertSharedSlate
 			}
 			
 			Delegate(Object);
-			if (ObjectUtils::IsActor(Object))
+			if (ConcertSyncCore::IsActor(Object))
 			{
 				AddedActors.Add(Object);
 			}
@@ -421,7 +421,7 @@ namespace UE::ConcertSharedSlate
 		TSet<TSharedPtr<FReplicatedObjectData>> NonRootNodes;
 		for (const TSharedPtr<FReplicatedObjectData>& Node : AllObjectRowData)
 		{
-			if (ObjectUtils::IsActor(Node->GetObjectPath()))
+			if (ConcertSyncCore::IsActor(Node->GetObjectPath()))
 			{
 				RootObjectRowData.Add(Node);
 			}
@@ -443,8 +443,8 @@ namespace UE::ConcertSharedSlate
 		}
 
 		// Find top level object of ReplicatedObjectData
-		const FSoftObjectPath OwningActor = ObjectUtils::GetActorOf(ObjectPath).Get(ObjectPath);
-		if (!ObjectUtils::IsActor(OwningActor))
+		const FSoftObjectPath OwningActor = ConcertSyncCore::GetActorOf(ObjectPath).Get(ObjectPath);
+		if (!ConcertSyncCore::IsActor(OwningActor))
 		{
 			return;
 		}
@@ -498,7 +498,7 @@ namespace UE::ConcertSharedSlate
 		const TOptional<IObjectHierarchyModel::FParentInfo> ParentInfo = ObjectHierarchy ? ObjectHierarchy->GetParentInfo(Object) : TOptional<IObjectHierarchyModel::FParentInfo>{};
 		// If no hierarchy was provided during construction, only show actors. If hierarchy provided, check whether this type of subobject is allowed.
 		const bool bCanShowWithinHierarchy = (ParentInfo && ShouldDisplayObjectRelation(ParentInfo->Relationship))
-			|| ObjectUtils::IsActor(Object.GetUniqueID());
+			|| ConcertSyncCore::IsActor(Object.GetUniqueID());
 		const bool bDidDelegateAllow = !ShouldDisplayObjectDelegate.IsBound() || ShouldDisplayObjectDelegate.Execute(Object.GetUniqueID());
 		return bCanShowWithinHierarchy && bDidDelegateAllow;
 	}
diff --git a/Engine/Plugins/Developer/Concert/ConcertUI/ConcertSharedSlate/Source/ConcertSharedSlate/Public/Replication/Utils/ObjectUtils.h b/Engine/Plugins/Developer/Concert/ConcertUI/ConcertSharedSlate/Source/ConcertSharedSlate/Public/Replication/Utils/ObjectUtils.h
deleted file mode 100644
