Generalize SClientName widget to use TAttribute<FText> for the content placed into the parentheses instead of TAttribute<bool>.
Context: Previously this widget would display "DisplayName(You)" or "DisplayName". This allows you to display any FText in the parentheses.

Also moves the types into ConcertSharedSlate namespace from ConcertSharedClientSlate namespace; it was in the wrong namespace in the first place.

diff --git a/Engine/Plugins/Developer/Concert/ConcertApp/MultiUserClient/Source/MultiUserClient/Private/Widgets/ActiveSession/Overview/SActiveSessionOverviewTab.cpp b/Engine/Plugins/Developer/Concert/ConcertApp/MultiUserClient/Source/MultiUserClient/Private/Widgets/ActiveSession/Overview/SActiveSessionOverviewTab.cpp
--- a/Engine/Plugins/Developer/Concert/ConcertApp/MultiUserClient/Source/MultiUserClient/Private/Widgets/ActiveSession/Overview/SActiveSessionOverviewTab.cpp
+++ b/Engine/Plugins/Developer/Concert/ConcertApp/MultiUserClient/Source/MultiUserClient/Private/Widgets/ActiveSession/Overview/SActiveSessionOverviewTab.cpp
@@ -79,9 +79,10 @@ public:
 			}
 
 			const bool bIsLocalClient = SyncClientPin->GetConcertClient()->GetCurrentSession()->GetSessionClientEndpointId() == SessionClientInfoPin->ClientEndpointId;
-			return SNew(UE::ConcertClientSharedSlate::SClientName)
+			const FText ParenthesesContent = bIsLocalClient ? UE::ConcertSharedSlate::ParenthesesClientNameContent::LocalClient : FText::GetEmpty();
+			return SNew(UE::ConcertSharedSlate::SClientName)
 				.ClientInfo(UE::ConcertClientSharedSlate::MakeClientInfoAttribute(SyncClientPin->GetConcertClient(), SessionClientInfoPin->ClientEndpointId))
-				.DisplayAsLocalClient(bIsLocalClient);
+				.ParenthesisContent(ParenthesesContent);
 		}
 		else if (ColumnName == ActiveSessionDetailsUI::PresenceColumnName)
 		{
diff --git a/Engine/Plugins/Developer/Concert/ConcertApp/MultiUserClient/Source/MultiUserClient/Private/Widgets/ActiveSession/Replication/Client/ClientUtils.cpp b/Engine/Plugins/Developer/Concert/ConcertApp/MultiUserClient/Source/MultiUserClient/Private/Widgets/ActiveSession/Replication/Client/ClientUtils.cpp
--- a/Engine/Plugins/Developer/Concert/ConcertApp/MultiUserClient/Source/MultiUserClient/Private/Widgets/ActiveSession/Replication/Client/ClientUtils.cpp
+++ b/Engine/Plugins/Developer/Concert/ConcertApp/MultiUserClient/Source/MultiUserClient/Private/Widgets/ActiveSession/Replication/Client/ClientUtils.cpp
@@ -20,13 +20,13 @@ namespace UE::MultiUserClient::Replication::ClientUtils
 		const bool bIsLocalClient = InSession.GetSessionClientEndpointId() == InClientEndpointId;
 		if (bIsLocalClient)
 		{
-			return ConcertClientSharedSlate::SClientName::GetDisplayText(InSession.GetLocalClientInfo(), bIsLocalClient).ToString();
+			return ConcertSharedSlate::SClientName::GetDisplayText(InSession.GetLocalClientInfo(), bIsLocalClient).ToString();
 		}
 
 		FConcertSessionClientInfo ClientInfo;
 		if (InSession.FindSessionClient(InClientEndpointId, ClientInfo))
 		{
-			return ConcertClientSharedSlate::SClientName::GetDisplayText(ClientInfo.ClientInfo, bIsLocalClient).ToString();
+			return ConcertSharedSlate::SClientName::GetDisplayText(ClientInfo.ClientInfo, bIsLocalClient).ToString();
 		}
 
 		ensureMsgf(false, TEXT("Bad args"));
@@ -81,8 +81,8 @@ namespace UE::MultiUserClient::Replication::ClientUtils
 		{
 			const FConcertClientInfo& LeftInfo = ClientToDisplayInfo[&Left];
 			const FConcertClientInfo& RightInfo = ClientToDisplayInfo[&Right];
-			const FText LeftDisplayName = ConcertClientSharedSlate::SClientName::GetDisplayText(LeftInfo, false);
-			const FText RightDisplayName = ConcertClientSharedSlate::SClientName::GetDisplayText(RightInfo, false);
+			const FText LeftDisplayName = ConcertSharedSlate::SClientName::GetDisplayText(LeftInfo);
+			const FText RightDisplayName = ConcertSharedSlate::SClientName::GetDisplayText(RightInfo);
 			return LeftDisplayName.ToString() <= RightDisplayName.ToString();
 		});
 		
diff --git a/Engine/Plugins/Developer/Concert/ConcertApp/MultiUserClient/Source/MultiUserClient/Private/Widgets/ActiveSession/Replication/Client/Multi/Columns/MultiStreamColumns_AssignedClients.cpp b/Engine/Plugins/Developer/Concert/ConcertApp/MultiUserClient/Source/MultiUserClient/Private/Widgets/ActiveSession/Replication/Client/Multi/Columns/MultiStreamColumns_AssignedClients.cpp
--- a/Engine/Plugins/Developer/Concert/ConcertApp/MultiUserClient/Source/MultiUserClient/Private/Widgets/ActiveSession/Replication/Client/Multi/Columns/MultiStreamColumns_AssignedClients.cpp
+++ b/Engine/Plugins/Developer/Concert/ConcertApp/MultiUserClient/Source/MultiUserClient/Private/Widgets/ActiveSession/Replication/Client/Multi/Columns/MultiStreamColumns_AssignedClients.cpp
@@ -54,7 +54,7 @@ namespace UE::MultiUserClient::Replication::MultiStreamColumns
 		
 		class SAssignedClientsWidget : public SCompoundWidget
 		{
-			TSharedPtr<ConcertClientSharedSlate::SHorizontalClientList> ClientList;
+			TSharedPtr<ConcertSharedSlate::SHorizontalClientList> ClientList;
 			FSoftObjectPath ManagedObject;
 			const ConcertSharedSlate::IObjectHierarchyModel* ObjectHierarchy = nullptr;
 			FReassignObjectPropertiesLogic* ReassignmentLogic = nullptr;
@@ -79,8 +79,8 @@ namespace UE::MultiUserClient::Replication::MultiStreamColumns
 			
 				ChildSlot
 				[
-					SAssignNew(ClientList, ConcertClientSharedSlate::SHorizontalClientList)
-					.IsLocalClient(ConcertClientSharedSlate::MakeIsLocalClientGetter(InConcertClient))
+					SAssignNew(ClientList, ConcertSharedSlate::SHorizontalClientList)
+					.GetClientParenthesesContent(ConcertClientSharedSlate::MakeGetLocalClientParenthesesContent(InConcertClient))
 					.GetClientInfo(ConcertClientSharedSlate::MakeClientInfoGetter(InConcertClient))
 					.Font(FAppStyle::GetFontStyle("PropertyWindow.NormalFont"))
 					.HighlightText(InArgs._HighlightText)
@@ -185,14 +185,20 @@ namespace UE::MultiUserClient::Replication::MultiStreamColumns
 			
 			TOptional<FString> GetDisplayString(const FSoftObjectPath& ManagedObject) const
 			{
-				using SWidgetType = ConcertClientSharedSlate::SHorizontalClientList;
+				using SWidgetType = ConcertSharedSlate::SHorizontalClientList;
 				const TArray<FGuid> Clients = Private::GetDisplayedClients(ObjectHierarchy, ReassignmentLogic, ManagedObject);
-				const ConcertSharedSlate::FIsLocalClient IsLocalClientDelegate = ConcertClientSharedSlate::MakeIsLocalClientGetter(ConcertClient);
+
+				const ConcertSharedSlate::FGetClientParenthesesContent GetParenthesesContent =
+					ConcertClientSharedSlate::MakeGetLocalClientParenthesesContent(ConcertClient);
+				const auto SortPredicate = [&GetParenthesesContent](const FConcertSessionClientInfo& Left, const FConcertSessionClientInfo& Right)
+				{
+					return ConcertSharedSlate::SortLocalClientParenthesesFirstThenThenAlphabetical(Left, Right, GetParenthesesContent);
+				};
 				return SWidgetType::GetDisplayString(
 					Clients,
 					ConcertClientSharedSlate::MakeClientInfoGetter(ConcertClient),
-					SWidgetType::FSortPredicate::CreateStatic(&SWidgetType::SortLocalClientFirstThenAlphabetical, IsLocalClientDelegate),
-					IsLocalClientDelegate
+					ConcertSharedSlate::FClientSortPredicate::CreateLambda(SortPredicate),
+					GetParenthesesContent
 					);
 			}
 		};
diff --git a/Engine/Plugins/Developer/Concert/ConcertApp/MultiUserClient/Source/MultiUserClient/Private/Widgets/ActiveSession/Replication/Client/Multi/Columns/SAssignPropertyComboBox.cpp b/Engine/Plugins/Developer/Concert/ConcertApp/MultiUserClient/Source/MultiUserClient/Private/Widgets/ActiveSession/Replication/Client/Multi/Columns/SAssignPropertyComboBox.cpp
--- a/Engine/Plugins/Developer/Concert/ConcertApp/MultiUserClient/Source/MultiUserClient/Private/Widgets/ActiveSession/Replication/Client/Multi/Columns/SAssignPropertyComboBox.cpp
+++ b/Engine/Plugins/Developer/Concert/ConcertApp/MultiUserClient/Source/MultiUserClient/Private/Widgets/ActiveSession/Replication/Client/Multi/Columns/SAssignPropertyComboBox.cpp
@@ -57,14 +57,21 @@ namespace UE::MultiUserClient::Replication
 		const FConcertPropertyChain& DisplayedProperty,
 		const TArray<TSoftObjectPtr<>>& EditedObjects)
 	{
-		using SWidgetType = ConcertClientSharedSlate::SHorizontalClientList;
+		using SWidgetType = ConcertSharedSlate::SHorizontalClientList;
 		const TArray<FGuid> Clients = AssignPropertyComboBox::GetDisplayedClients(ClientManager, DisplayedProperty, EditedObjects);
-		const ConcertSharedSlate::FIsLocalClient IsLocalClientDelegate = ConcertClientSharedSlate::MakeIsLocalClientGetter(LocalConcertClient);
+		
+		const ConcertSharedSlate::FGetClientParenthesesContent GetParenthesesContent =
+			ConcertClientSharedSlate::MakeGetLocalClientParenthesesContent(LocalConcertClient);
+		const auto SortPredicate = [&GetParenthesesContent](const FConcertSessionClientInfo& Left, const FConcertSessionClientInfo& Right)
+		{
+			return ConcertSharedSlate::SortLocalClientParenthesesFirstThenThenAlphabetical(Left, Right, GetParenthesesContent);
+		};
+		
 		return SWidgetType::GetDisplayString(
 			Clients,
 			ConcertClientSharedSlate::MakeClientInfoGetter(LocalConcertClient),
-			SWidgetType::FSortPredicate::CreateStatic(&SWidgetType::SortLocalClientFirstThenAlphabetical, IsLocalClientDelegate),
-			IsLocalClientDelegate
+			ConcertSharedSlate::FClientSortPredicate::CreateLambda(SortPredicate),
+			GetParenthesesContent
 			);
 	}
 
@@ -91,8 +98,8 @@ namespace UE::MultiUserClient::Replication
 			.HasDownArrow(true)
 			.ButtonContent()
 			[
-				SAssignNew(ClientListWidget, ConcertClientSharedSlate::SHorizontalClientList)
-				.IsLocalClient(ConcertClientSharedSlate::MakeIsLocalClientGetter(ConcertClient.ToSharedRef()))
+				SAssignNew(ClientListWidget, ConcertSharedSlate::SHorizontalClientList)
+				.GetClientParenthesesContent(ConcertClientSharedSlate::MakeGetLocalClientParenthesesContent(ConcertClient.ToSharedRef()))
 				.GetClientInfo(ConcertClientSharedSlate::MakeClientInfoGetter(ConcertClient.ToSharedRef()))
 				.Font(FAppStyle::GetFontStyle("PropertyWindow.NormalFont"))
 				.HighlightText_Lambda([this](){ return HighlightText ? *HighlightText : FText::GetEmpty(); })
@@ -115,7 +122,7 @@ namespace UE::MultiUserClient::Replication
 
 	TSharedRef<SWidget> SAssignPropertyComboBox::GetMenuContent()
 	{
-		using namespace ConcertClientSharedSlate;
+		using namespace ConcertSharedSlate;
 
 		const auto MakeWidget = [this](const FGuid& EndpointId) -> TSharedRef<SWidget>
 		{
diff --git a/Engine/Plugins/Developer/Concert/ConcertApp/MultiUserClient/Source/MultiUserClient/Private/Widgets/ActiveSession/Replication/Client/Multi/Columns/SAssignPropertyComboBox.h b/Engine/Plugins/Developer/Concert/ConcertApp/MultiUserClient/Source/MultiUserClient/Private/Widgets/ActiveSession/Replication/Client/Multi/Columns/SAssignPropertyComboBox.h
--- a/Engine/Plugins/Developer/Concert/ConcertApp/MultiUserClient/Source/MultiUserClient/Private/Widgets/ActiveSession/Replication/Client/Multi/Columns/SAssignPropertyComboBox.h
+++ b/Engine/Plugins/Developer/Concert/ConcertApp/MultiUserClient/Source/MultiUserClient/Private/Widgets/ActiveSession/Replication/Client/Multi/Columns/SAssignPropertyComboBox.h
@@ -13,8 +13,11 @@ class IConcertClient;
 enum class ECheckBoxState : uint8;
 struct FGuid;
 
-namespace UE::ConcertClientSharedSlate { class SHorizontalClientList; }
-namespace UE::ConcertSharedSlate { class IMultiReplicationStreamEditor; }
+namespace UE::ConcertSharedSlate
+{
+	class IMultiReplicationStreamEditor;
+	class SHorizontalClientList;
+}
 
 namespace UE::MultiUserClient::Replication
 {
@@ -70,7 +73,7 @@ namespace UE::MultiUserClient::Replication
 		FConcertPropertyChain Property;
 
 		/** The static menu content (when there is no drop-down). */
-		TSharedPtr<ConcertClientSharedSlate::SHorizontalClientList> ClientListWidget;
+		TSharedPtr<ConcertSharedSlate::SHorizontalClientList> ClientListWidget;
 		/** Passed to MakeListWidgetDelegate */
 		TSharedPtr<FText> HighlightText;
 
diff --git a/Engine/Plugins/Developer/Concert/ConcertApp/MultiUserClient/Source/MultiUserClient/Private/Widgets/ActiveSession/Replication/Client/SPresetComboButton.cpp b/Engine/Plugins/Developer/Concert/ConcertApp/MultiUserClient/Source/MultiUserClient/Private/Widgets/ActiveSession/Replication/Client/SPresetComboButton.cpp
--- a/Engine/Plugins/Developer/Concert/ConcertApp/MultiUserClient/Source/MultiUserClient/Private/Widgets/ActiveSession/Replication/Client/SPresetComboButton.cpp
+++ b/Engine/Plugins/Developer/Concert/ConcertApp/MultiUserClient/Source/MultiUserClient/Private/Widgets/ActiveSession/Replication/Client/SPresetComboButton.cpp
@@ -136,7 +136,7 @@ namespace UE::MultiUserClient::Replication
 
 		const auto GetDisplayText = [](const FConcertClientInfo& ClientInfo, bool bDisplayAsLocalClient = false)
 		{
-			return ConcertClientSharedSlate::SClientName::GetDisplayText(ClientInfo, bDisplayAsLocalClient);
+			return ConcertSharedSlate::SClientName::GetDisplayText(ClientInfo, bDisplayAsLocalClient);
 		};
 		
 		TArray<FConcertSessionClientInfo> RemoteClients = Session->GetSessionClients();
diff --git a/Engine/Plugins/Developer/Concert/ConcertApp/MultiUserClient/Source/MultiUserClient/Private/Widgets/Disconnected/SConcertClientSessionBrowser.cpp b/Engine/Plugins/Developer/Concert/ConcertApp/MultiUserClient/Source/MultiUserClient/Private/Widgets/Disconnected/SConcertClientSessionBrowser.cpp
--- a/Engine/Plugins/Developer/Concert/ConcertApp/MultiUserClient/Source/MultiUserClient/Private/Widgets/Disconnected/SConcertClientSessionBrowser.cpp
+++ b/Engine/Plugins/Developer/Concert/ConcertApp/MultiUserClient/Source/MultiUserClient/Private/Widgets/Disconnected/SConcertClientSessionBrowser.cpp
@@ -295,9 +295,8 @@ TSharedRef<SWidget> SConcertClientSessionBrowser::MakeUserAndSettings(TSharedPtr
 		.VAlign(VAlign_Center)
 		.AutoWidth()
 		[
-			SNew(UE::ConcertClientSharedSlate::SClientName)
+			SNew(UE::ConcertSharedSlate::SClientName)
 			.ClientInfo(UE::ConcertClientSharedSlate::MakeLocalClientInfoAttribute(Client))
-			.DisplayAsLocalClient(false) // It is obvious that it's the local user because we have not joined any session yet so do not display (You) in the back.
 			.HighlightText_Lambda([InSearchText](){ return *InSearchText.Get(); })
 		]
 
diff --git a/Engine/Plugins/Developer/Concert/ConcertUI/ConcertClientSharedSlate/Source/ConcertClientSharedSlate/Public/Widgets/Client/ClientInfoHelpers.h b/Engine/Plugins/Developer/Concert/ConcertUI/ConcertClientSharedSlate/Source/ConcertClientSharedSlate/Public/Widgets/Client/ClientInfoHelpers.h
--- a/Engine/Plugins/Developer/Concert/ConcertUI/ConcertClientSharedSlate/Source/ConcertClientSharedSlate/Public/Widgets/Client/ClientInfoHelpers.h
+++ b/Engine/Plugins/Developer/Concert/ConcertUI/ConcertClientSharedSlate/Source/ConcertClientSharedSlate/Public/Widgets/Client/ClientInfoHelpers.h
@@ -6,6 +6,7 @@
 #include "IConcertClient.h"
 #include "IConcertSession.h"
 #include "Widgets/Client/ClientInfoDelegate.h"
+#include "Widgets/Client/SClientName.h"
 
 #include "Misc/Attribute.h"
 #include "Misc/Optional.h"
@@ -50,6 +51,19 @@ namespace UE::ConcertClientSharedSlate
 		});
 	}
 
+	/** @return A parentheses delegate that returns "You" if the endpoint ID is that of Client and returns FText::GetEmpty() otherwise. */
+	inline ConcertSharedSlate::FGetClientParenthesesContent MakeGetLocalClientParenthesesContent(const TSharedRef<IConcertClient>& Client)
+	{
+		return ConcertSharedSlate::FGetClientParenthesesContent::CreateLambda([WeakClient = Client.ToWeakPtr()](const FGuid& ClientEndpointId)
+		{
+			const TSharedPtr<IConcertClient> ClientPin = WeakClient.Pin();
+			const bool bIsLocalClient =  ClientPin
+				&& ClientPin->GetCurrentSession()
+				&& ClientPin->GetCurrentSession()->GetSessionClientEndpointId() == ClientEndpointId;
+			return bIsLocalClient ? ConcertSharedSlate::ParenthesesClientNameContent::LocalClient : FText::GetEmpty();
+		});
+	}
+
 	inline TAttribute<TOptional<FConcertClientInfo>> MakeLocalClientInfoAttribute(const TSharedRef<IConcertClient>& Client)
 	{
 		return TAttribute<TOptional<FConcertClientInfo>>::CreateLambda([WeakClient = Client.ToWeakPtr()]()
diff --git a/Engine/Plugins/Developer/Concert/ConcertUI/ConcertSharedSlate/Source/ConcertSharedSlate/Private/Widgets/Client/SClientName.cpp b/Engine/Plugins/Developer/Concert/ConcertUI/ConcertSharedSlate/Source/ConcertSharedSlate/Private/Widgets/Client/SClientName.cpp
--- a/Engine/Plugins/Developer/Concert/ConcertUI/ConcertSharedSlate/Source/ConcertSharedSlate/Private/Widgets/Client/SClientName.cpp
+++ b/Engine/Plugins/Developer/Concert/ConcertUI/ConcertSharedSlate/Source/ConcertSharedSlate/Private/Widgets/Client/SClientName.cpp
@@ -10,14 +10,18 @@
 
 #define LOCTEXT_NAMESPACE "SClientName"
 
-namespace UE::ConcertClientSharedSlate
+namespace UE::ConcertSharedSlate::ParenthesesClientNameContent
+{
+	const FText LocalClient = LOCTEXT("ParenthesesClientNameContent.LocalClient", "You");
+}
+
+namespace UE::ConcertSharedSlate
 {
 	void SClientName::Construct(const FArguments& InArgs)
 	{
 		ClientInfoAttribute = InArgs._ClientInfo;
-		DisplayAsLocalClientAttribute = InArgs._DisplayAsLocalClient;
+		ParenthesisContentAttribute = InArgs._ParenthesisContent;
 		check(ClientInfoAttribute.IsSet() || ClientInfoAttribute.IsBound());
-		check(DisplayAsLocalClientAttribute.IsSet() || DisplayAsLocalClientAttribute.IsBound());
 		
 		ChildSlot
 		[
@@ -54,11 +58,19 @@ namespace UE::ConcertClientSharedSlate
 
 	FText SClientName::GetDisplayText(const FConcertClientInfo& Info, bool bDisplayAsLocalClient)
 	{
-		if (bDisplayAsLocalClient)
+		return bDisplayAsLocalClient
+			? GetDisplayText(Info, ParenthesesClientNameContent::LocalClient)
+			: GetDisplayText(Info);
+	}
+
+	FText SClientName::GetDisplayText(const FConcertClientInfo& Info, const FText& ParenthesesContent)
+	{
+		if (!ParenthesesContent.IsEmpty())
 		{
 			return FText::Format(
-				LOCTEXT("ClientDisplayNameFmt", "{0} (You)"),
-				FText::FromString(Info.DisplayName)
+				LOCTEXT("ClientDisplayNameFmt", "{0} ({1})"),
+				FText::FromString(Info.DisplayName),
+				ParenthesesContent
 				);
 		}
 		
@@ -68,8 +80,11 @@ namespace UE::ConcertClientSharedSlate
 	FText SClientName::GetClientDisplayName() const
 	{
 		const TOptional<FConcertClientInfo> ClientInfo = ClientInfoAttribute.Get();
+		const FText ParenthesesContent = ParenthesisContentAttribute.IsSet() || ParenthesisContentAttribute.IsBound()
+			? ParenthesisContentAttribute.Get()
+			: FText::GetEmpty();
 		return ClientInfo
-			? GetDisplayText(*ClientInfo, DisplayAsLocalClientAttribute.Get())
+			? GetDisplayText(*ClientInfo, ParenthesesContent)
 			: LOCTEXT("Unavailable", "Unavailable");
 	}
 
diff --git a/Engine/Plugins/Developer/Concert/ConcertUI/ConcertSharedSlate/Source/ConcertSharedSlate/Private/Widgets/Client/SHorizontalClientList.cpp b/Engine/Plugins/Developer/Concert/ConcertUI/ConcertSharedSlate/Source/ConcertSharedSlate/Private/Widgets/Client/SHorizontalClientList.cpp
--- a/Engine/Plugins/Developer/Concert/ConcertUI/ConcertSharedSlate/Source/ConcertSharedSlate/Private/Widgets/Client/SHorizontalClientList.cpp
+++ b/Engine/Plugins/Developer/Concert/ConcertUI/ConcertSharedSlate/Source/ConcertSharedSlate/Private/Widgets/Client/SHorizontalClientList.cpp
@@ -12,49 +12,40 @@
 
 #define LOCTEXT_NAMESPACE "SHorizontalClientList"
 
-namespace UE::ConcertClientSharedSlate
+namespace UE::ConcertSharedSlate::HorizontalClientList
 {
-	namespace HorizontalClientList
-	{
-		TArray<FConcertSessionClientInfo> GetSortedClients(
+	static TArray<FConcertSessionClientInfo> GetSortedClients(
 			const TConstArrayView<FGuid>& Clients,
 			const ConcertSharedSlate::FGetOptionalClientInfo& GetClientInfoDelegate,
-			const SHorizontalClientList::FSortPredicate& SortPredicate
+			const FClientSortPredicate& SortPredicate
 			)
+	{
+		// Prefetch the client info to avoid many GetClientInfoDelegate calls during Sort()
+		TArray<FConcertSessionClientInfo> ClientsToDisplay;
+		for (const FGuid& Client : Clients)
 		{
-			// Prefetch the client info to avoid many GetClientInfoDelegate calls during Sort()
-			TArray<FConcertSessionClientInfo> ClientsToDisplay;
-			for (const FGuid& Client : Clients)
-			{
-				const TOptional<FConcertClientInfo> ClientInfo = GetClientInfoDelegate.Execute(Client);
-				FConcertSessionClientInfo Info;
-				Info.ClientEndpointId = Client;
-				Info.ClientInfo = ClientInfo ? *ClientInfo : FConcertClientInfo{ .DisplayName = TEXT("Unavailable") };
-				ClientsToDisplay.Add(Info);
-			}
-
-			ClientsToDisplay.Sort([&SortPredicate](const FConcertSessionClientInfo& Left, const FConcertSessionClientInfo& Right)
-			{
-				return SortPredicate.Execute(Left, Right);
-			});
-			return ClientsToDisplay;
+			const TOptional<FConcertClientInfo> ClientInfo = GetClientInfoDelegate.Execute(Client);
+			FConcertSessionClientInfo Info;
+			Info.ClientEndpointId = Client;
+			Info.ClientInfo = ClientInfo ? *ClientInfo : FConcertClientInfo{ .DisplayName = TEXT("Unavailable") };
+			ClientsToDisplay.Add(Info);
 		}
+
+		ClientsToDisplay.Sort([&SortPredicate](const FConcertSessionClientInfo& Left, const FConcertSessionClientInfo& Right)
+		{
+			return SortPredicate.Execute(Left, Right);
+		});
+		return ClientsToDisplay;
 	}
-	
-	bool SHorizontalClientList::SortLocalClientFirstThenAlphabetical(const FConcertSessionClientInfo& Left, const FConcertSessionClientInfo& Right, ConcertSharedSlate::FIsLocalClient IsLocalClientDelegate)
-	{
-		// If one of the compare clients is local, always return that the local client is smaller.
-		const bool bLeftIsLocalClient = IsLocalClientDelegate.IsBound() && IsLocalClientDelegate.Execute(Left.ClientEndpointId);
-		const bool bRightIsLocalClient = IsLocalClientDelegate.IsBound() && IsLocalClientDelegate.Execute(Right.ClientEndpointId);
-		return bLeftIsLocalClient
-			|| (!bRightIsLocalClient && Left.ClientInfo.DisplayName < Right.ClientInfo.DisplayName);
-	}
+}
 
+namespace UE::ConcertSharedSlate
+{
 	TOptional<FString> SHorizontalClientList::GetDisplayString(
 		const TConstArrayView<FGuid>& Clients,
-		const ConcertSharedSlate::FGetOptionalClientInfo& GetClientInfoDelegate,
-		const FSortPredicate& SortPredicate,
-		const ConcertSharedSlate::FIsLocalClient& IsLocalClientDelegate
+		const FGetOptionalClientInfo& GetClientInfoDelegate,
+		const FClientSortPredicate& SortPredicate,
+		const FGetClientParenthesesContent& GetClientParenthesesContent
 		)
 	{
 		const TArray<FConcertSessionClientInfo> ClientsToDisplay = HorizontalClientList::GetSortedClients(Clients, GetClientInfoDelegate, SortPredicate);
@@ -63,23 +54,28 @@ namespace UE::ConcertClientSharedSlate
 			return {};
 		}
 		
-		return FString::JoinBy(ClientsToDisplay, TEXT(", "), [&IsLocalClientDelegate](const FConcertSessionClientInfo& ClientInfo)
+		return FString::JoinBy(ClientsToDisplay, TEXT(", "), [&GetClientParenthesesContent](const FConcertSessionClientInfo& ClientInfo)
 			{
 				// GetSortedClients should return empty if GetSortedClients is invalid
-				const bool bIsLocalClient = IsLocalClientDelegate.IsBound() && IsLocalClientDelegate.Execute(ClientInfo.ClientEndpointId);
-				return SClientName::GetDisplayText(ClientInfo.ClientInfo, bIsLocalClient).ToString();
+				const FText ParenthesesContent = EvaluateGetClientParenthesesContent(GetClientParenthesesContent, ClientInfo.ClientEndpointId);
+				return SClientName::GetDisplayText(ClientInfo.ClientInfo, ParenthesesContent).ToString();
 			});
 	}
 
 	void SHorizontalClientList::Construct(const FArguments& InArgs)
 	{
-		IsLocalClientDelegate = InArgs._IsLocalClient;
+		GetClientParenthesesContentDelegate = InArgs._GetClientParenthesesContent;
 		GetClientInfoDelegate = InArgs._GetClientInfo;
 		SortPredicateDelegate = InArgs._SortPredicate.IsBound()
 			? InArgs._SortPredicate
-			: FSortPredicate::CreateStatic(&SHorizontalClientList::SortLocalClientFirstThenAlphabetical, IsLocalClientDelegate);
+			: FClientSortPredicate::CreateLambda([this](const FConcertSessionClientInfo& Left, const FConcertSessionClientInfo& Right)
+			{
+				return SortSpecifiedParenthesesFirstThenThenAlphabetical(
+					Left, Right, GetClientParenthesesContentDelegate, ParenthesesClientNameContent::LocalClient
+					);
+			});
 		
-		DisplayAvatarColorAttribute = InArgs._DisplayAvatarColor;
+		ShouldDisplayAvatarColorAttribute = InArgs._DisplayAvatarColor;
 		HighlightTextAttribute = InArgs._HighlightText;
 		
 		NameFont = InArgs._Font;
@@ -103,7 +99,7 @@ namespace UE::ConcertClientSharedSlate
 		];
 	}
 
-	void SHorizontalClientList::RefreshList(const TConstArrayView<FGuid>& Clients)
+	void SHorizontalClientList::RefreshList(const TConstArrayView<FGuid>& Clients) const
 	{
 		ScrollBox->ClearChildren();
 
@@ -115,7 +111,6 @@ namespace UE::ConcertClientSharedSlate
 		WidgetSwitcher->SetActiveWidgetIndex(1);
 
 		const TArray<FConcertSessionClientInfo> ClientsToDisplay = HorizontalClientList::GetSortedClients(Clients, GetClientInfoDelegate, SortPredicateDelegate);
-		
 		bool bIsFirst = true;
 		for (const FConcertSessionClientInfo& Info : ClientsToDisplay)
 		{
@@ -133,30 +128,19 @@ namespace UE::ConcertClientSharedSlate
 			}
 
 			const FGuid& EndpointId = Info.ClientEndpointId;
-			const bool bIsLocalClient = IsLocalClientDelegate.IsBound() && IsLocalClientDelegate.Execute(Info.ClientEndpointId);
-			if (bIsLocalClient)
-			{
-				ScrollBox->AddSlot()
-				[
-					SNew(SLocalClientName)
-					.DisplayInfo_Lambda([this, EndpointId](){ return GetClientInfoDelegate.Execute(EndpointId); })
-					.DisplayAvatarColor(DisplayAvatarColorAttribute)
-					.HighlightText(HighlightTextAttribute)
-					.Font(NameFont)
-				];
-			}
-			else
+			const auto GetParenthesesContent = [this, EndpointId]()
 			{
-				ScrollBox->AddSlot()
-				[
-					SNew(SRemoteClientName)
-					.DisplayInfo_Lambda([this, EndpointId](){ return GetClientInfoDelegate.Execute(EndpointId); })
-					.DisplayAvatarColor(DisplayAvatarColorAttribute)
-					.HighlightText(HighlightTextAttribute)
-					.Font(NameFont)
-				];
-			}
-
+				return EvaluateGetClientParenthesesContent(GetClientParenthesesContentDelegate, EndpointId);
+			};
+			ScrollBox->AddSlot()
+			[
+				SNew(SClientName)
+				.ClientInfo_Lambda([this, EndpointId](){ return GetClientInfoDelegate.Execute(EndpointId); })
+				.ParenthesisContent_Lambda(GetParenthesesContent)
+				.DisplayAvatarColor(ShouldDisplayAvatarColorAttribute)
+				.HighlightText(HighlightTextAttribute)
+				.Font(NameFont)
+			];
 			bIsFirst = false;
 		}
 	}
diff --git a/Engine/Plugins/Developer/Concert/ConcertUI/ConcertSharedSlate/Source/ConcertSharedSlate/Private/Widgets/Client/SLocalClientName.cpp b/Engine/Plugins/Developer/Concert/ConcertUI/ConcertSharedSlate/Source/ConcertSharedSlate/Private/Widgets/Client/SLocalClientName.cpp
--- a/Engine/Plugins/Developer/Concert/ConcertUI/ConcertSharedSlate/Source/ConcertSharedSlate/Private/Widgets/Client/SLocalClientName.cpp
+++ b/Engine/Plugins/Developer/Concert/ConcertUI/ConcertSharedSlate/Source/ConcertSharedSlate/Private/Widgets/Client/SLocalClientName.cpp
@@ -7,7 +7,7 @@
 
 #define LOCTEXT_NAMESPACE "SLocalClientName"
 
-namespace UE::ConcertClientSharedSlate
+namespace UE::ConcertSharedSlate
 {
 	void SLocalClientName::Construct(const FArguments& InArgs)
 	{
@@ -15,7 +15,7 @@ namespace UE::ConcertClientSharedSlate
 		[
 			SNew(SClientName)
 			.ClientInfo(InArgs._DisplayInfo)
-			.DisplayAsLocalClient(true)
+			.ParenthesisContent(ParenthesesClientNameContent::LocalClient)
 			.DisplayAvatarColor(InArgs._DisplayAvatarColor)
 			.HighlightText(InArgs._HighlightText)
 			.Font(InArgs._Font)
diff --git a/Engine/Plugins/Developer/Concert/ConcertUI/ConcertSharedSlate/Source/ConcertSharedSlate/Private/Widgets/Client/SRemoteClientName.cpp b/Engine/Plugins/Developer/Concert/ConcertUI/ConcertSharedSlate/Source/ConcertSharedSlate/Private/Widgets/Client/SRemoteClientName.cpp
--- a/Engine/Plugins/Developer/Concert/ConcertUI/ConcertSharedSlate/Source/ConcertSharedSlate/Private/Widgets/Client/SRemoteClientName.cpp
+++ b/Engine/Plugins/Developer/Concert/ConcertUI/ConcertSharedSlate/Source/ConcertSharedSlate/Private/Widgets/Client/SRemoteClientName.cpp
@@ -7,7 +7,7 @@
 
 #define LOCTEXT_NAMESPACE "SLocalClientName"
 
-namespace UE::ConcertClientSharedSlate
+namespace UE::ConcertSharedSlate
 {
 	void SRemoteClientName::Construct(const FArguments& InArgs)
 	{
@@ -17,7 +17,6 @@ namespace UE::ConcertClientSharedSlate
 		[
 			SNew(SClientName)
 			.ClientInfo(this, &SRemoteClientName::GetClientInfo)
-			.DisplayAsLocalClient(false)
 			.DisplayAvatarColor(InArgs._DisplayAvatarColor)
 			.HighlightText(InArgs._HighlightText)
 			.Font(InArgs._Font)
diff --git a/Engine/Plugins/Developer/Concert/ConcertUI/ConcertSharedSlate/Source/ConcertSharedSlate/Public/Widgets/Client/ClientInfoDelegate.h b/Engine/Plugins/Developer/Concert/ConcertUI/ConcertSharedSlate/Source/ConcertSharedSlate/Public/Widgets/Client/ClientInfoDelegate.h
--- a/Engine/Plugins/Developer/Concert/ConcertUI/ConcertSharedSlate/Source/ConcertSharedSlate/Public/Widgets/Client/ClientInfoDelegate.h
+++ b/Engine/Plugins/Developer/Concert/ConcertUI/ConcertSharedSlate/Source/ConcertSharedSlate/Public/Widgets/Client/ClientInfoDelegate.h
@@ -12,4 +12,14 @@ namespace UE::ConcertSharedSlate
 {
 	DECLARE_DELEGATE_RetVal_OneParam(bool, FIsLocalClient, const FGuid& ClientId);
 	DECLARE_DELEGATE_RetVal_OneParam(TOptional<FConcertClientInfo>, FGetOptionalClientInfo, const FGuid& ClientEndpointId);
+	/**
+	 * Used by SHorizontalList to determine what should go in the parentheses behind the display name, e.g. "Display Name (You)".
+	 * Return FText::GetEmpty() if nothing should go into the parentheses.
+	 */
+	DECLARE_DELEGATE_RetVal_OneParam(FText, FGetClientParenthesesContent, const FGuid& ClientId);
+
+	FORCEINLINE FText EvaluateGetClientParenthesesContent(const FGetClientParenthesesContent& Getter, const FGuid& ClientId)
+	{
+		return Getter.IsBound() ? Getter.Execute(ClientId) : FText::GetEmpty();
+	}
 }
\ No newline at end of file
diff --git a/Engine/Plugins/Developer/Concert/ConcertUI/ConcertSharedSlate/Source/ConcertSharedSlate/Public/Widgets/Client/ClientSorting.h b/Engine/Plugins/Developer/Concert/ConcertUI/ConcertSharedSlate/Source/ConcertSharedSlate/Public/Widgets/Client/ClientSorting.h
new file mode 100644
--- /dev/null
+++ b/Engine/Plugins/Developer/Concert/ConcertUI/ConcertSharedSlate/Source/ConcertSharedSlate/Public/Widgets/Client/ClientSorting.h
@@ -0,0 +1,62 @@
+﻿// Copyright Epic Games, Inc. All Rights Reserved.
+
+#pragma once
+
+#include "ClientInfoDelegate.h"
+#include "SClientName.h"
+#include "Styling/AppStyle.h"
+#include "Widgets/DeclarativeSyntaxSupport.h"
+#include "Widgets/SCompoundWidget.h"
+
+class IConcertClient;
+class SScrollBox;
+class SWidgetSwitcher;
+
+struct FConcertSessionClientInfo;
+struct FGuid;
+
+namespace UE::ConcertSharedSlate
+{
+	DECLARE_DELEGATE_RetVal_TwoParams(bool,
+		FClientSortPredicate,
+		const FConcertSessionClientInfo& Left,
+		const FConcertSessionClientInfo& Right
+		);
+
+	template<typename TDelegate>
+	concept CClientInfoSortable = std::is_convertible_v<TDelegate, const FGetClientParenthesesContent&>;
+
+	/**
+	 * Predicate for sorting everything that has ParenthesesContentToPlaceFirst first, and then sorts everything alphabetically.
+	 *
+	 * This is supposed to be used with SHorizontalClientList.
+	 * Example input list: ["AClient", "BClient", "ZClient(You)"]
+	 * Example output list sorted by ParenthesesContentToPlaceFirst == "You": ZClient(You), AClient, BClient 
+	 */
+	inline bool SortSpecifiedParenthesesFirstThenThenAlphabetical(
+		const FConcertSessionClientInfo& Left,
+		const FConcertSessionClientInfo& Right,
+		const FGetClientParenthesesContent& GetClientParenthesesContent,
+		const FText& ParenthesesContentToPlaceFirst
+		)
+	{
+		const bool bLeftHasParenthesesContent = GetClientParenthesesContent.IsBound()
+			&& GetClientParenthesesContent.Execute(Left.ClientEndpointId).EqualTo(ParenthesesContentToPlaceFirst);
+		const bool bRightHasParenthesesContent = GetClientParenthesesContent.IsBound()
+			&& GetClientParenthesesContent.Execute(Right.ClientEndpointId).EqualTo(ParenthesesContentToPlaceFirst);
+		return bLeftHasParenthesesContent
+			|| (!bRightHasParenthesesContent && Left.ClientInfo.DisplayName < Right.ClientInfo.DisplayName);
+	}
+
+	/** Sorts anything that has "You" in the parentheses first. */
+	inline bool SortLocalClientParenthesesFirstThenThenAlphabetical(
+		const FConcertSessionClientInfo& Left,
+		const FConcertSessionClientInfo& Right,
+		const FGetClientParenthesesContent& GetClientParenthesesContent
+		)
+	{
+		return SortSpecifiedParenthesesFirstThenThenAlphabetical(
+			Left, Right, GetClientParenthesesContent, ParenthesesClientNameContent::LocalClient
+			);
+	}
+}
\ No newline at end of file
diff --git a/Engine/Plugins/Developer/Concert/ConcertUI/ConcertSharedSlate/Source/ConcertSharedSlate/Public/Widgets/Client/SClientName.h b/Engine/Plugins/Developer/Concert/ConcertUI/ConcertSharedSlate/Source/ConcertSharedSlate/Public/Widgets/Client/SClientName.h
--- a/Engine/Plugins/Developer/Concert/ConcertUI/ConcertSharedSlate/Source/ConcertSharedSlate/Public/Widgets/Client/SClientName.h
+++ b/Engine/Plugins/Developer/Concert/ConcertUI/ConcertSharedSlate/Source/ConcertSharedSlate/Public/Widgets/Client/SClientName.h
@@ -9,23 +9,36 @@
 
 class IConcertClient;
 
-namespace UE::ConcertClientSharedSlate
+/** Contains SClientName's parentheses content definitions. */
+namespace UE::ConcertSharedSlate::ParenthesesClientNameContent
 {
-	/** Knows how to display FConcertClientInfo. */
+	/** The client corresponds to the local user. "You" is appended to name, e.g. "ClientName(You)". */
+	CONCERTSHAREDSLATE_API extern const FText LocalClient;
+}
+
+namespace UE::ConcertSharedSlate
+{
+	/**
+	 * Knows how to display FConcertClientInfo.
+	 * 
+	 * The widget looks like this: []DisplayName(ParenthesesContent)
+	 *  - [] is a square displaying the avatar colour (optional)
+	 *  - DisplayName is FConcertClientInfo::DisplayName
+	 *  - ParenthesesContent is additional info you can supply, like "You" (optional). See ParenthesesClientNameContent. 
+	 */
 	class CONCERTSHAREDSLATE_API SClientName : public SCompoundWidget
 	{
 	public:
 
 		SLATE_BEGIN_ARGS(SClientName)
-			: _DisplayAsLocalClient(false)
-			, _DisplayAvatarColor(true)
+			: _DisplayAvatarColor(true)
 			, _Font(FAppStyle::Get().GetFontStyle("BoldFont"))
 		{}
 			/** The client info to display. */
 			SLATE_ATTRIBUTE(TOptional<FConcertClientInfo>, ClientInfo)
 			
-			/** Whether visually indicate that this is a local client (appends "(me)" if true). */
-			SLATE_ATTRIBUTE(bool, DisplayAsLocalClient)
+			/** Content to display behind the display name in parentheses. */
+			SLATE_ATTRIBUTE(FText, ParenthesisContent)
 			/** Whether to show a square image in front of the name. */
 			SLATE_ATTRIBUTE(bool, DisplayAvatarColor)
 			
@@ -39,13 +52,15 @@ namespace UE::ConcertClientSharedSlate
 
 		/** @return The display that would be used given the settings. */
 		static FText GetDisplayText(const FConcertClientInfo& Info, bool bDisplayAsLocalClient);
+		/** @return The display that would be used given the settings. */
+		static FText GetDisplayText(const FConcertClientInfo& Info, const FText& ParenthesesContent = FText::GetEmpty());
 
 	private:
 
 		/** The client info to display. */
 		TAttribute<TOptional<FConcertClientInfo>> ClientInfoAttribute;
-		/** Whether visually indicate that this is a local client (appends "(me)" if true). */
-		TAttribute<bool> DisplayAsLocalClientAttribute;
+		/** Content to display behind the display name in parentheses. */
+		TAttribute<FText> ParenthesisContentAttribute;
 
 		/** Gets the display name. */
 		FText GetClientDisplayName() const;
diff --git a/Engine/Plugins/Developer/Concert/ConcertUI/ConcertSharedSlate/Source/ConcertSharedSlate/Public/Widgets/Client/SHorizontalClientList.h b/Engine/Plugins/Developer/Concert/ConcertUI/ConcertSharedSlate/Source/ConcertSharedSlate/Public/Widgets/Client/SHorizontalClientList.h
--- a/Engine/Plugins/Developer/Concert/ConcertUI/ConcertSharedSlate/Source/ConcertSharedSlate/Public/Widgets/Client/SHorizontalClientList.h
+++ b/Engine/Plugins/Developer/Concert/ConcertUI/ConcertSharedSlate/Source/ConcertSharedSlate/Public/Widgets/Client/SHorizontalClientList.h
@@ -3,6 +3,7 @@
 #pragma once
 
 #include "ClientInfoDelegate.h"
+#include "ClientSorting.h"
 #include "Styling/AppStyle.h"
 #include "Widgets/DeclarativeSyntaxSupport.h"
 #include "Widgets/SCompoundWidget.h"
@@ -14,29 +15,27 @@ class SWidgetSwitcher;
 struct FConcertSessionClientInfo;
 struct FGuid;
 
-namespace UE::ConcertClientSharedSlate
+namespace UE::ConcertSharedSlate
 {
 	/** Aligns client widgets from left to right. If there is not enough space, a horizontal scroll bar cuts of the list. */
 	class CONCERTSHAREDSLATE_API SHorizontalClientList : public SCompoundWidget
 	{
 	public:
 
-		DECLARE_DELEGATE_RetVal_TwoParams(bool, FSortPredicate, const FConcertSessionClientInfo& Left, const FConcertSessionClientInfo& FConcertSessionClientInfo);
-		static bool SortLocalClientFirstThenAlphabetical(const FConcertSessionClientInfo& Left, const FConcertSessionClientInfo& Right, ConcertSharedSlate::FIsLocalClient IsLocalClientDelegate);
-
 		/** @return The display string a SHorizontalClientList would display with the given state. Returns unset optional if EmptyListSlot would be shown. */
 		static TOptional<FString> GetDisplayString(
 			const TConstArrayView<FGuid>& Clients,
-			const ConcertSharedSlate::FGetOptionalClientInfo& GetClientInfoDelegate,
-			const FSortPredicate& SortPredicate,
-			const ConcertSharedSlate::FIsLocalClient& IsLocalClientDelegate = {}
+			const FGetOptionalClientInfo& GetClientInfoDelegate,
+			const FClientSortPredicate& SortPredicate,
+			const FGetClientParenthesesContent& GetClientParenthesesContent = {}
 			);
 		
 		SLATE_BEGIN_ARGS(SHorizontalClientList)
 			: _Font(FAppStyle::Get().GetFontStyle("NormalFont"))
 		{}
-			/** Decides whether the given client should be displayed as if it was a local client. */
-			SLATE_EVENT(ConcertSharedSlate::FIsLocalClient, IsLocalClient)
+			/** Gets the content to place in parentheses behind the given client. */
+			SLATE_EVENT(ConcertSharedSlate::FGetClientParenthesesContent, GetClientParenthesesContent)
+			
 			/** Used to get client display info for remote clients. */
 			SLATE_EVENT(ConcertSharedSlate::FGetOptionalClientInfo, GetClientInfo)
 			
@@ -49,7 +48,7 @@ namespace UE::ConcertClientSharedSlate
 			SLATE_ARGUMENT(FSlateFontInfo, Font)
 			
 			/** Defaults to placing the local client first (if contained) and sorting alphabetically otherwise. */
-			SLATE_EVENT(FSortPredicate, SortPredicate)
+			SLATE_EVENT(FClientSortPredicate, SortPredicate)
 
 			/** Tooltip text to display when the list is non-empty. */
 			SLATE_ATTRIBUTE(FText, ListToolTipText)
@@ -61,19 +60,19 @@ namespace UE::ConcertClientSharedSlate
 		void Construct(const FArguments& InArgs);
 
 		/** Refreshes the list. */
-		void RefreshList(const TConstArrayView<FGuid>& Clients);
+		void RefreshList(const TConstArrayView<FGuid>& Clients) const;
 
 	private:
 		
-		/** Decides whether the given client should be displayed as if it was a local client. */
-		ConcertSharedSlate::FIsLocalClient IsLocalClientDelegate;
+		/** Gets the content to place in parentheses behind the given client. */
+		FGetClientParenthesesContent GetClientParenthesesContentDelegate;
 		/** Used to get client display info for remote clients. */
-		ConcertSharedSlate::FGetOptionalClientInfo GetClientInfoDelegate;
+		FGetOptionalClientInfo GetClientInfoDelegate;
 		/** Sorts the client list */
-		FSortPredicate SortPredicateDelegate;
+		FClientSortPredicate SortPredicateDelegate;
 
 		/** Whether the square in front of the client name should be displayed. */
-		TAttribute<bool> DisplayAvatarColorAttribute;
+		TAttribute<bool> ShouldDisplayAvatarColorAttribute;
 		/** Used for highlighting in the text */
 		TAttribute<FText> HighlightTextAttribute;
 		
diff --git a/Engine/Plugins/Developer/Concert/ConcertUI/ConcertSharedSlate/Source/ConcertSharedSlate/Public/Widgets/Client/SLocalClientName.h b/Engine/Plugins/Developer/Concert/ConcertUI/ConcertSharedSlate/Source/ConcertSharedSlate/Public/Widgets/Client/SLocalClientName.h
--- a/Engine/Plugins/Developer/Concert/ConcertUI/ConcertSharedSlate/Source/ConcertSharedSlate/Public/Widgets/Client/SLocalClientName.h
+++ b/Engine/Plugins/Developer/Concert/ConcertUI/ConcertSharedSlate/Source/ConcertSharedSlate/Public/Widgets/Client/SLocalClientName.h
@@ -10,7 +10,7 @@
 
 class IConcertClient;
 
-namespace UE::ConcertClientSharedSlate
+namespace UE::ConcertSharedSlate
 {
 	/**
 	 * Displays the name of a client.
diff --git a/Engine/Plugins/Developer/Concert/ConcertUI/ConcertSharedSlate/Source/ConcertSharedSlate/Public/Widgets/Client/SRemoteClientName.h b/Engine/Plugins/Developer/Concert/ConcertUI/ConcertSharedSlate/Source/ConcertSharedSlate/Public/Widgets/Client/SRemoteClientName.h
--- a/Engine/Plugins/Developer/Concert/ConcertUI/ConcertSharedSlate/Source/ConcertSharedSlate/Public/Widgets/Client/SRemoteClientName.h
+++ b/Engine/Plugins/Developer/Concert/ConcertUI/ConcertSharedSlate/Source/ConcertSharedSlate/Public/Widgets/Client/SRemoteClientName.h
@@ -10,7 +10,7 @@
 
 class IConcertClient;
 
-namespace UE::ConcertClientSharedSlate
+namespace UE::ConcertSharedSlate
 {
 	/**
 	 * Displays the name of a client.
