The properties column no longer counts properties double if multiple clients have them assigned; becomes relevant for offline clients, for which the properties would be counted twice.

diff --git a/Engine/Plugins/Developer/Concert/ConcertUI/ConcertSharedSlate/Source/ConcertSharedSlate/Private/Replication/Editor/View/MultiEditor/ConsolidatedMultiStreamModel.cpp b/Engine/Plugins/Developer/Concert/ConcertUI/ConcertSharedSlate/Source/ConcertSharedSlate/Private/Replication/Editor/View/MultiEditor/ConsolidatedMultiStreamModel.cpp
--- a/Engine/Plugins/Developer/Concert/ConcertUI/ConcertSharedSlate/Source/ConcertSharedSlate/Private/Replication/Editor/View/MultiEditor/ConsolidatedMultiStreamModel.cpp
+++ b/Engine/Plugins/Developer/Concert/ConcertUI/ConcertSharedSlate/Source/ConcertSharedSlate/Private/Replication/Editor/View/MultiEditor/ConsolidatedMultiStreamModel.cpp
@@ -124,13 +124,27 @@ namespace UE::ConcertSharedSlate
 
 	uint32 FConsolidatedMultiStreamModel::GetNumProperties(const FSoftObjectPath& Object) const
 	{
-		uint32 NumProperties = 0;
-		MultiStreamModel->ForEachStream([&Object, &NumProperties](const TSharedRef<IReplicationStreamModel>& Model)
+		uint32 PropertyUpperLimit = 0;
+		MultiStreamModel->ForEachStream([&Object, &PropertyUpperLimit](const TSharedRef<IReplicationStreamModel>& Model)
 		{
-			NumProperties += Model->GetNumProperties(Object);
+			PropertyUpperLimit += Model->GetNumProperties(Object);
 			return EBreakBehavior::Continue;
 		});
-		return NumProperties;
+
+		// This is needed, so we don't count properties double (two separate streams may contain the same property twice).
+		TSet<FConcertPropertyChain> Properties;
+		Properties.Reserve(PropertyUpperLimit);
+		MultiStreamModel->ForEachStream([&Object, &Properties](const TSharedRef<IReplicationStreamModel>& Model)
+		{
+			Model->ForEachProperty(Object, [&Properties](const FConcertPropertyChain& Property)
+			{
+				Properties.Add(Property);
+				return EBreakBehavior::Continue;
+			});
+			return EBreakBehavior::Continue;
+		});
+		
+		return Properties.Num();
 	}
 
 	void FConsolidatedMultiStreamModel::AddObjects(TConstArrayView<UObject*> Objects)
