Fix: Renaming actor does not update Replication label in Multi-User Browser

diff --git a/Engine/Plugins/Developer/Concert/ConcertUI/ConcertSharedSlate/Source/ConcertSharedSlate/Private/Replication/Editor/View/SelectionViewerColumns.cpp b/Engine/Plugins/Developer/Concert/ConcertUI/ConcertSharedSlate/Source/ConcertSharedSlate/Private/Replication/Editor/View/SelectionViewerColumns.cpp
--- a/Engine/Plugins/Developer/Concert/ConcertUI/ConcertSharedSlate/Source/ConcertSharedSlate/Private/Replication/Editor/View/SelectionViewerColumns.cpp
+++ b/Engine/Plugins/Developer/Concert/ConcertUI/ConcertSharedSlate/Source/ConcertSharedSlate/Private/Replication/Editor/View/SelectionViewerColumns.cpp
@@ -53,7 +53,6 @@ namespace UE::ConcertSharedSlate::ReplicationColumns::TopLevel
 				const FReplicatedObjectData& ObjectData = InArgs.RowItem.RowData;
 				const TSoftObjectPtr<>& Object = ObjectData.GetObjectPtr();
 				
-				const FText Text = GetDisplayText(ObjectData);
 				const FSlateIcon ClassIcon = GetObjectClassDelegate.IsBound() ? DisplayUtils::GetObjectIcon(GetObjectClassDelegate.Execute(Object)) : FSlateIcon{};
 				return SNew(SHorizontalBox)
 					.ToolTipText(FText::FromString(Object.ToString()))
@@ -74,7 +73,7 @@ namespace UE::ConcertSharedSlate::ReplicationColumns::TopLevel
 					[
 						SNew(STextBlock)
 						.HighlightText(TAttribute<FText>::CreateLambda([HighlightText = InArgs.HighlightText](){ return *HighlightText; }))
-						.Text(Text)
+						.Text_Lambda([this, ObjectPtr = ObjectData.GetObjectPtr()](){ return GetDisplayText(ObjectPtr); })
 					];
 			}
 			
@@ -89,7 +88,7 @@ namespace UE::ConcertSharedSlate::ReplicationColumns::TopLevel
 			virtual bool CanBeSorted() const override { return true; } 
 			virtual bool IsLessThan(const FObjectTreeRowContext& Left, const FObjectTreeRowContext& Right) const override
 			{
-				return GetDisplayText(Left.RowData).ToString() < GetDisplayText(Right.RowData).ToString();
+				return GetDisplayText(Left.RowData.GetObjectPtr()).ToString() < GetDisplayText(Right.RowData.GetObjectPtr()).ToString();
 			}
 
 		private:
@@ -97,9 +96,9 @@ namespace UE::ConcertSharedSlate::ReplicationColumns::TopLevel
 			IObjectNameModel* const OptionalNameModel;
 			const FGetObjectClass GetObjectClassDelegate;
 			
-			FText GetDisplayText(const FReplicatedObjectData& ObjectData) const
+			FText GetDisplayText(const TSoftObjectPtr<>& ObjectPtr) const
 			{
-				return DisplayUtils::GetObjectDisplayText(ObjectData.GetObjectPtr(), OptionalNameModel);
+				return DisplayUtils::GetObjectDisplayText(ObjectPtr, OptionalNameModel);
 			}
 		};
 
