Extend FCreateViewerParams with FWrapOutlinerWidget, which allows API users to extend the replication outliner, e.g. with a SDropTarget.

diff --git a/Engine/Plugins/Developer/Concert/ConcertUI/ConcertSharedSlate/Source/ConcertSharedSlate/Private/Replication/Editor/View/ObjectEditor/SBaseReplicationStreamEditor.cpp b/Engine/Plugins/Developer/Concert/ConcertUI/ConcertSharedSlate/Source/ConcertSharedSlate/Private/Replication/Editor/View/ObjectEditor/SBaseReplicationStreamEditor.cpp
--- a/Engine/Plugins/Developer/Concert/ConcertUI/ConcertSharedSlate/Source/ConcertSharedSlate/Private/Replication/Editor/View/ObjectEditor/SBaseReplicationStreamEditor.cpp
+++ b/Engine/Plugins/Developer/Concert/ConcertUI/ConcertSharedSlate/Source/ConcertSharedSlate/Private/Replication/Editor/View/ObjectEditor/SBaseReplicationStreamEditor.cpp
@@ -78,6 +78,7 @@ namespace UE::ConcertSharedSlate
 						return FHoverRowContent{ OverlayWidget, Alignment };
 					})
 				.NoOutlinerObjects(LOCTEXT("NoObjects", "Add objects to replicate"))
+				.WrapOutliner(InArgs._WrapOutliner)
 		];
 	}
 
diff --git a/Engine/Plugins/Developer/Concert/ConcertUI/ConcertSharedSlate/Source/ConcertSharedSlate/Private/Replication/Editor/View/ObjectEditor/SBaseReplicationStreamEditor.h b/Engine/Plugins/Developer/Concert/ConcertUI/ConcertSharedSlate/Source/ConcertSharedSlate/Private/Replication/Editor/View/ObjectEditor/SBaseReplicationStreamEditor.h
--- a/Engine/Plugins/Developer/Concert/ConcertUI/ConcertSharedSlate/Source/ConcertSharedSlate/Private/Replication/Editor/View/ObjectEditor/SBaseReplicationStreamEditor.h
+++ b/Engine/Plugins/Developer/Concert/ConcertUI/ConcertSharedSlate/Source/ConcertSharedSlate/Private/Replication/Editor/View/ObjectEditor/SBaseReplicationStreamEditor.h
@@ -93,6 +93,9 @@ namespace UE::ConcertSharedSlate
 			SLATE_ATTRIBUTE(bool, IsEditingEnabled)
 			/** Optional. Whenever IsEditingEnabled returns true, this tooltip is displayed for relevant, disabled UI. */
 			SLATE_ATTRIBUTE(FText, EditingDisabledToolTipText)
+			
+			/** Optional. Returns new widget that wraps the outliner widget, which displays the replicated objects. Could return e.g. an SOverlay or SDropTarget. */
+			SLATE_EVENT(FWrapOutlinerWidget, WrapOutliner)
 		SLATE_END_ARGS()
 
 		void Construct(const FArguments& InArgs,
diff --git a/Engine/Plugins/Developer/Concert/ConcertUI/ConcertSharedSlate/Source/ConcertSharedSlate/Private/Replication/Editor/View/ObjectViewer/SReplicationStreamViewer.cpp b/Engine/Plugins/Developer/Concert/ConcertUI/ConcertSharedSlate/Source/ConcertSharedSlate/Private/Replication/Editor/View/ObjectViewer/SReplicationStreamViewer.cpp
--- a/Engine/Plugins/Developer/Concert/ConcertUI/ConcertSharedSlate/Source/ConcertSharedSlate/Private/Replication/Editor/View/ObjectViewer/SReplicationStreamViewer.cpp
+++ b/Engine/Plugins/Developer/Concert/ConcertUI/ConcertSharedSlate/Source/ConcertSharedSlate/Private/Replication/Editor/View/ObjectViewer/SReplicationStreamViewer.cpp
@@ -279,7 +279,7 @@ namespace UE::ConcertSharedSlate
 					];
 			}();
 		
-		return SAssignNew(ReplicatedObjects, SReplicationTreeView<FReplicatedObjectData>)
+		TSharedRef<SWidget> Outliner = SAssignNew(ReplicatedObjects, SReplicationTreeView<FReplicatedObjectData>)
 			.RootItemsSource(&RootObjectRowData)
 			.OnGetChildren(this, &SReplicationStreamViewer::GetObjectRowChildren)
 			.OnContextMenuOpening(InArgs._OnObjectsContextMenuOpening)
@@ -301,6 +301,10 @@ namespace UE::ConcertSharedSlate
 			.NoItemsContent() [ SNew(STextBlock).Text(NoObjectsAttribute) ]
 			.GetHoveredRowContent(InArgs._GetHoveredRowContent)
 			.RowStyle(FAppStyle::Get(), "TableView.AlternatingRow");
+
+		return InArgs._WrapOutliner.IsBound()
+			? InArgs._WrapOutliner.Execute(Outliner)
+			: Outliner;
 	}
 
 	TSharedRef<SWidget> SReplicationStreamViewer::CreatePropertiesSection(const FArguments& InArgs)
diff --git a/Engine/Plugins/Developer/Concert/ConcertUI/ConcertSharedSlate/Source/ConcertSharedSlate/Private/Replication/Editor/View/ObjectViewer/SReplicationStreamViewer.h b/Engine/Plugins/Developer/Concert/ConcertUI/ConcertSharedSlate/Source/ConcertSharedSlate/Private/Replication/Editor/View/ObjectViewer/SReplicationStreamViewer.h
--- a/Engine/Plugins/Developer/Concert/ConcertUI/ConcertSharedSlate/Source/ConcertSharedSlate/Private/Replication/Editor/View/ObjectViewer/SReplicationStreamViewer.h
+++ b/Engine/Plugins/Developer/Concert/ConcertUI/ConcertSharedSlate/Source/ConcertSharedSlate/Private/Replication/Editor/View/ObjectViewer/SReplicationStreamViewer.h
@@ -87,6 +87,9 @@ namespace UE::ConcertSharedSlate
 
 			/** Optional text to display when no object is in the outliner. Defaults to "No objects to display." "*/
 			SLATE_ATTRIBUTE(FText, NoOutlinerObjects)
+
+			/** Optional. Returns new widget that wraps the outliner widget, which displays the replicated objects. Could return e.g. an SOverlay or SDropTarget. */
+			SLATE_EVENT(FWrapOutlinerWidget, WrapOutliner)
 		SLATE_END_ARGS()
 
 		void Construct(const FArguments& InArgs, const TSharedRef<IReplicationStreamModel>& InPropertiesModel);
diff --git a/Engine/Plugins/Developer/Concert/ConcertUI/ConcertSharedSlate/Source/ConcertSharedSlate/Private/Replication/ReplicationWidgetFactories.cpp b/Engine/Plugins/Developer/Concert/ConcertUI/ConcertSharedSlate/Source/ConcertSharedSlate/Private/Replication/ReplicationWidgetFactories.cpp
--- a/Engine/Plugins/Developer/Concert/ConcertUI/ConcertSharedSlate/Source/ConcertSharedSlate/Private/Replication/ReplicationWidgetFactories.cpp
+++ b/Engine/Plugins/Developer/Concert/ConcertUI/ConcertSharedSlate/Source/ConcertSharedSlate/Private/Replication/ReplicationWidgetFactories.cpp
@@ -48,7 +48,8 @@ namespace UE::ConcertSharedSlate
 			.LeftOfObjectSearchBar() [ MoveTemp(ViewerParams.LeftOfObjectSearchBar.Widget) ]
 			.RightOfObjectSearchBar() [ MoveTemp(ViewerParams.RightOfObjectSearchBar.Widget) ]
 			.IsEditingEnabled(MoveTemp(EditorParams.IsEditingEnabled))
-			.EditingDisabledToolTipText(MoveTemp(EditorParams.EditingDisabledToolTipText));
+			.EditingDisabledToolTipText(MoveTemp(EditorParams.EditingDisabledToolTipText))
+			.WrapOutliner(MoveTemp(ViewerParams.WrapOutlinerWidgetDelegate));
 	}
 
 	TSharedRef<IPropertyTreeView> CreateSearchablePropertyTreeView(FCreatePropertyTreeViewParams Params)
diff --git a/Engine/Plugins/Developer/Concert/ConcertUI/ConcertSharedSlate/Source/ConcertSharedSlate/Public/Replication/ReplicationWidgetFactories.h b/Engine/Plugins/Developer/Concert/ConcertUI/ConcertSharedSlate/Source/ConcertSharedSlate/Public/Replication/ReplicationWidgetFactories.h
--- a/Engine/Plugins/Developer/Concert/ConcertUI/ConcertSharedSlate/Source/ConcertSharedSlate/Public/Replication/ReplicationWidgetFactories.h
+++ b/Engine/Plugins/Developer/Concert/ConcertUI/ConcertSharedSlate/Source/ConcertSharedSlate/Public/Replication/ReplicationWidgetFactories.h
@@ -100,6 +100,9 @@ namespace UE::ConcertSharedSlate
 		FMakeObjectRowOverlayWidget MakeObjectRowOverlayWidgetDelegate;
 		/** Optional. If MakeRowOverlayWidgetDelegate is specified, this controls how the widget is aligned in the column. */
 		EHorizontalAlignment OverlayWidgetAlignment = HAlign_Right;
+
+		/** Optional. Returns a new widget that wraps the outliner widget, which displays the replicated objects. Could return e.g. an SOverlay or SDropTarget. */
+		FWrapOutlinerWidget WrapOutlinerWidgetDelegate;
 	};
 
 	/** Params for creating an IReplicationStreamEditor */
diff --git a/Engine/Plugins/Developer/Concert/ConcertUI/ConcertSharedSlate/Source/ConcertSharedSlate/Public/Replication/Utils/ReplicationWidgetDelegates.h b/Engine/Plugins/Developer/Concert/ConcertUI/ConcertSharedSlate/Source/ConcertSharedSlate/Public/Replication/Utils/ReplicationWidgetDelegates.h
--- a/Engine/Plugins/Developer/Concert/ConcertUI/ConcertSharedSlate/Source/ConcertSharedSlate/Public/Replication/Utils/ReplicationWidgetDelegates.h
+++ b/Engine/Plugins/Developer/Concert/ConcertUI/ConcertSharedSlate/Source/ConcertSharedSlate/Public/Replication/Utils/ReplicationWidgetDelegates.h
@@ -41,4 +41,10 @@ namespace UE::ConcertSharedSlate
 
 	/** Delegate executed to create a widget that overlays an object row. */
 	DECLARE_DELEGATE_RetVal_OneParam(TSharedRef<SWidget>, FMakeObjectRowOverlayWidget, const FReplicatedObjectData& ObjectData);
+
+	/**
+	 * Delegate executed during construction of UI. It receives the SReplicationTreeView showing the replicated objects.
+	 * The returned widget is intended to wrap the tree in some way, e.g. by it into an SOverlay or SDropTarget.
+	 */
+	DECLARE_DELEGATE_RetVal_OneParam(TSharedRef<SWidget>, FWrapOutlinerWidget, const TSharedRef<SWidget>& OutlinerWidget);
 }
\ No newline at end of file
