Introduce IReplicationStreamViewer::SetSelectedObjects.

diff --git a/Engine/Plugins/Developer/Concert/ConcertUI/ConcertSharedSlate/Source/ConcertSharedSlate/Private/Replication/Editor/View/ObjectEditor/SBaseReplicationStreamEditor.cpp b/Engine/Plugins/Developer/Concert/ConcertUI/ConcertSharedSlate/Source/ConcertSharedSlate/Private/Replication/Editor/View/ObjectEditor/SBaseReplicationStreamEditor.cpp
--- a/Engine/Plugins/Developer/Concert/ConcertUI/ConcertSharedSlate/Source/ConcertSharedSlate/Private/Replication/Editor/View/ObjectEditor/SBaseReplicationStreamEditor.cpp
+++ b/Engine/Plugins/Developer/Concert/ConcertUI/ConcertSharedSlate/Source/ConcertSharedSlate/Private/Replication/Editor/View/ObjectEditor/SBaseReplicationStreamEditor.cpp
@@ -103,6 +103,11 @@ namespace UE::ConcertSharedSlate
 		ReplicationViewer->RequestPropertyColumnResort(ColumnId);
 	}
 
+	void SBaseReplicationStreamEditor::SetSelectedObjects(TConstArrayView<TSoftObjectPtr<>> Objects)
+	{
+		ReplicationViewer->SetSelectedObjects(Objects);
+	}
+
 	TArray<TSoftObjectPtr<>> SBaseReplicationStreamEditor::GetSelectedObjects() const
 	{
 		return ReplicationViewer->GetSelectedObjects();
diff --git a/Engine/Plugins/Developer/Concert/ConcertUI/ConcertSharedSlate/Source/ConcertSharedSlate/Private/Replication/Editor/View/ObjectEditor/SBaseReplicationStreamEditor.h b/Engine/Plugins/Developer/Concert/ConcertUI/ConcertSharedSlate/Source/ConcertSharedSlate/Private/Replication/Editor/View/ObjectEditor/SBaseReplicationStreamEditor.h
--- a/Engine/Plugins/Developer/Concert/ConcertUI/ConcertSharedSlate/Source/ConcertSharedSlate/Private/Replication/Editor/View/ObjectEditor/SBaseReplicationStreamEditor.h
+++ b/Engine/Plugins/Developer/Concert/ConcertUI/ConcertSharedSlate/Source/ConcertSharedSlate/Private/Replication/Editor/View/ObjectEditor/SBaseReplicationStreamEditor.h
@@ -109,6 +109,7 @@ namespace UE::ConcertSharedSlate
 		virtual void Refresh() override;
 		virtual void RequestObjectColumnResort(const FName& ColumnId) override;
 		virtual void RequestPropertyColumnResort(const FName& ColumnId) override;
+		virtual void SetSelectedObjects(TConstArrayView<TSoftObjectPtr<>> Objects) override;
 		virtual TArray<TSoftObjectPtr<>> GetSelectedObjects() const override;
 		//~ End IReplicationStreamEditor Interface
 
diff --git a/Engine/Plugins/Developer/Concert/ConcertUI/ConcertSharedSlate/Source/ConcertSharedSlate/Private/Replication/Editor/View/ObjectViewer/SReplicationStreamViewer.cpp b/Engine/Plugins/Developer/Concert/ConcertUI/ConcertSharedSlate/Source/ConcertSharedSlate/Private/Replication/Editor/View/ObjectViewer/SReplicationStreamViewer.cpp
--- a/Engine/Plugins/Developer/Concert/ConcertUI/ConcertSharedSlate/Source/ConcertSharedSlate/Private/Replication/Editor/View/ObjectViewer/SReplicationStreamViewer.cpp
+++ b/Engine/Plugins/Developer/Concert/ConcertUI/ConcertSharedSlate/Source/ConcertSharedSlate/Private/Replication/Editor/View/ObjectViewer/SReplicationStreamViewer.cpp
@@ -71,6 +71,12 @@ namespace UE::ConcertSharedSlate
 		PropertySection->RequestResortForColumn(ColumnId);
 	}
 
+	void SReplicationStreamViewer::SetSelectedObjects(TConstArrayView<TSoftObjectPtr<>> Objects)
+	{
+		// Do the selection at the end of the tick in case anything still needs updating this tick.
+		SelectObjects(Objects, true);
+	}
+
 	TArray<TSoftObjectPtr<>> SReplicationStreamViewer::GetSelectedObjects() const
 	{
 		TArray<TSoftObjectPtr<>> Result;
diff --git a/Engine/Plugins/Developer/Concert/ConcertUI/ConcertSharedSlate/Source/ConcertSharedSlate/Private/Replication/Editor/View/ObjectViewer/SReplicationStreamViewer.h b/Engine/Plugins/Developer/Concert/ConcertUI/ConcertSharedSlate/Source/ConcertSharedSlate/Private/Replication/Editor/View/ObjectViewer/SReplicationStreamViewer.h
--- a/Engine/Plugins/Developer/Concert/ConcertUI/ConcertSharedSlate/Source/ConcertSharedSlate/Private/Replication/Editor/View/ObjectViewer/SReplicationStreamViewer.h
+++ b/Engine/Plugins/Developer/Concert/ConcertUI/ConcertSharedSlate/Source/ConcertSharedSlate/Private/Replication/Editor/View/ObjectViewer/SReplicationStreamViewer.h
@@ -98,6 +98,7 @@ namespace UE::ConcertSharedSlate
 		virtual void Refresh() override;
 		virtual void RequestObjectColumnResort(const FName& ColumnId) override;
 		virtual void RequestPropertyColumnResort(const FName& ColumnId) override;
+		virtual void SetSelectedObjects(TConstArrayView<TSoftObjectPtr<>> Objects) override;
 		virtual TArray<TSoftObjectPtr<>> GetSelectedObjects() const override;
 		//~ End IReplicationStreamViewer Interface
 
diff --git a/Engine/Plugins/Developer/Concert/ConcertUI/ConcertSharedSlate/Source/ConcertSharedSlate/Public/Replication/Editor/View/IReplicationStreamViewer.h b/Engine/Plugins/Developer/Concert/ConcertUI/ConcertSharedSlate/Source/ConcertSharedSlate/Public/Replication/Editor/View/IReplicationStreamViewer.h
--- a/Engine/Plugins/Developer/Concert/ConcertUI/ConcertSharedSlate/Source/ConcertSharedSlate/Public/Replication/Editor/View/IReplicationStreamViewer.h
+++ b/Engine/Plugins/Developer/Concert/ConcertUI/ConcertSharedSlate/Source/ConcertSharedSlate/Public/Replication/Editor/View/IReplicationStreamViewer.h
@@ -33,6 +33,8 @@ namespace UE::ConcertSharedSlate
 		UE_DEPRECATED(5.5, "Use GetSelectedObjects instead.")
 		virtual TArray<FSoftObjectPath> GetObjectsBeingPropertyEdited() const;
 
+		/** Selects the objects in the outliner and deselects all others. Only actors can be selected. Invalid entries are filtered out. */
+		virtual void SetSelectedObjects(TConstArrayView<TSoftObjectPtr<>> Objects) = 0;
 		/** @return The objects for which the properties are being edited / displayed. */
 		virtual TArray<TSoftObjectPtr<>> GetSelectedObjects() const = 0;
 
