Remove unnecessary check in the session browser on selection change.

diff --git a/Engine/Plugins/Developer/Concert/ConcertUI/ConcertSharedSlate/Source/ConcertSharedSlate/Private/Session/Browser/SConcertSessionBrowser.cpp b/Engine/Plugins/Developer/Concert/ConcertUI/ConcertSharedSlate/Source/ConcertSharedSlate/Private/Session/Browser/SConcertSessionBrowser.cpp
--- a/Engine/Plugins/Developer/Concert/ConcertUI/ConcertSharedSlate/Source/ConcertSharedSlate/Private/Session/Browser/SConcertSessionBrowser.cpp
+++ b/Engine/Plugins/Developer/Concert/ConcertUI/ConcertSharedSlate/Source/ConcertSharedSlate/Private/Session/Browser/SConcertSessionBrowser.cpp
@@ -1035,7 +1035,6 @@ void SConcertSessionBrowser::OnSessionSelectionChanged(TSharedPtr<FConcertSessio
 	// Cancel editing the row to create, archive or restore a session (if any), unless the row was selected in code.
 	if (EditableSessionRow.IsValid() && SelectInfo != ESelectInfo::Direct)
 	{
-		check (SelectedSession != EditableSessionRow); // User should not be able to reselect an editable row as we remove it as soon as it is unselected.
 		RemoveSessionRow(EditableSessionRow);
 		check(!EditableSessionRow.IsValid() && !EditableSessionRowParent.IsValid()); // Expect to be cleared by RemoveSessionRow().
 	}
