Updated headers using UnrealCodeFixup to make sure dllstorage is on methods/staticvars instead of types. Part 2/n

diff --git a/Engine/Plugins/Developer/Concert/ConcertMain/Source/ConcertServer/Public/ConcertServerSettings.h b/Engine/Plugins/Developer/Concert/ConcertMain/Source/ConcertServer/Public/ConcertServerSettings.h
--- a/Engine/Plugins/Developer/Concert/ConcertMain/Source/ConcertServer/Public/ConcertServerSettings.h
+++ b/Engine/Plugins/Developer/Concert/ConcertMain/Source/ConcertServer/Public/ConcertServerSettings.h
@@ -12,6 +12,8 @@
 #include "ConcertVersion.h"
 #include "ConcertServerSettings.generated.h"
 
+#define UE_API CONCERTSERVER_API
+
 USTRUCT()
 struct FConcertServerSettings
 {
@@ -26,12 +28,12 @@ struct FConcertServerSettings
 	int32 SessionTickFrequencySeconds = 1;
 };
 
-UCLASS(config=Engine)
-class CONCERTSERVER_API UConcertServerConfig : public UObject
+UCLASS(MinimalAPI, config=Engine)
+class UConcertServerConfig : public UObject
 {
 	GENERATED_BODY()
 public:
-	UConcertServerConfig();
+	UE_API UConcertServerConfig();
 
 	/**
 	 * If true, instruct the server to auto-archive sessions that were left in the working directory because the server did not exit properly rather than
@@ -123,3 +125,5 @@ public:
 	UPROPERTY(config)
 	bool bMountDefaultSessionRepository = true;
 };
+
+#undef UE_API
diff --git a/Engine/Plugins/Developer/Concert/ConcertSync/ConcertSyncClient/Source/ConcertSyncClient/Private/ConcertAssetContainer.h b/Engine/Plugins/Developer/Concert/ConcertSync/ConcertSyncClient/Source/ConcertSyncClient/Private/ConcertAssetContainer.h
--- a/Engine/Plugins/Developer/Concert/ConcertSync/ConcertSyncClient/Source/ConcertSyncClient/Private/ConcertAssetContainer.h
+++ b/Engine/Plugins/Developer/Concert/ConcertSync/ConcertSyncClient/Source/ConcertSyncClient/Private/ConcertAssetContainer.h
@@ -18,8 +18,8 @@ class UFont;
 /**
  * Asset container for VREditor.
  */
-UCLASS()
-class CONCERTSYNCCLIENT_API UConcertAssetContainer : public UDataAsset
+UCLASS(MinimalAPI)
+class UConcertAssetContainer : public UDataAsset
 {
 	GENERATED_BODY()
 
diff --git a/Engine/Plugins/Developer/Concert/ConcertSync/ConcertSyncClient/Source/ConcertSyncClient/Public/ConcertActivityStream.h b/Engine/Plugins/Developer/Concert/ConcertSync/ConcertSyncClient/Source/ConcertSyncClient/Public/ConcertActivityStream.h
--- a/Engine/Plugins/Developer/Concert/ConcertSync/ConcertSyncClient/Source/ConcertSyncClient/Public/ConcertActivityStream.h
+++ b/Engine/Plugins/Developer/Concert/ConcertSync/ConcertSyncClient/Source/ConcertSyncClient/Public/ConcertActivityStream.h
@@ -5,6 +5,8 @@
 #include "CoreMinimal.h"
 #include "ConcertMessageData.h"
 
+#define UE_API CONCERTSYNCCLIENT_API
+
 struct FConcertSessionActivity;
 class IConcertClient;
 
@@ -12,7 +14,7 @@ class IConcertClient;
  * Streams an archived session activities asynchronously and caches them until the client poll and consume them synchronously. The
  * streams also caches the corresponding client info for lookup.
  */
-class CONCERTSYNCCLIENT_API FConcertActivityStream
+class FConcertActivityStream
 {
 public:
 	/**
@@ -24,7 +26,7 @@ public:
 	 * @param InBufferSize The stream buffer cache size. The stream will fill its buffer and stop pulling them from the server until some activities are consumed.
 	 * @see FConcertClientSessionActivity
 	 */
-	FConcertActivityStream(TWeakPtr<IConcertClient, ESPMode::ThreadSafe> InConcertClient, const FGuid& InAdminEndpointId, const FGuid& InSessionId, bool bInIncludeActivityDetails, uint32 InBufferSize = 1024);
+	UE_API FConcertActivityStream(TWeakPtr<IConcertClient, ESPMode::ThreadSafe> InConcertClient, const FGuid& InAdminEndpointId, const FGuid& InSessionId, bool bInIncludeActivityDetails, uint32 InBufferSize = 1024);
 
 	/**
 	 * Reads all buffered activities. Call this function periodically (between frames, not in a tight loop). Every times this function is called,
@@ -35,14 +37,14 @@ public:
 	 * @param OutErrorMsg Contains the last error message, if any or an empty text.
 	 * @return True if all activities were consumed, false otherwise.
 	 */
-	bool Read(TArray<TSharedPtr<FConcertSessionActivity>>& InOutActivities, int32& OutReadCount, FText& OutErrorMsg);
+	UE_API bool Read(TArray<TSharedPtr<FConcertSessionActivity>>& InOutActivities, int32& OutReadCount, FText& OutErrorMsg);
 
 	/** Returns the client info corresponding to an activities. The stream maintains this state as an optimization. */
 	const FConcertClientInfo* GetActivityClientInfo(const FGuid& ActivityEndpointId) const { return EndpointClientInfoMap.Find(ActivityEndpointId); }
 
 private:
 	/** Run asynchronous requests in background and caches the activities until consumed.*/
-	void Fetch(int64 ActivityCount);
+	UE_API void Fetch(int64 ActivityCount);
 
 	/** Returns how many activities should be requested per request. */
 	constexpr int64 GetRequestBatchSize() { return BufferSize; }
@@ -62,3 +64,5 @@ private:
 	bool bFetching = false;
 	bool bIncludeActivityDetails = false;
 };
+
+#undef UE_API
diff --git a/Engine/Plugins/Developer/Concert/ConcertSync/ConcertSyncClient/Source/ConcertSyncClient/Public/Replication/IConcertClientReplicationManager.h b/Engine/Plugins/Developer/Concert/ConcertSync/ConcertSyncClient/Source/ConcertSyncClient/Public/Replication/IConcertClientReplicationManager.h
--- a/Engine/Plugins/Developer/Concert/ConcertSync/ConcertSyncClient/Source/ConcertSyncClient/Public/Replication/IConcertClientReplicationManager.h
+++ b/Engine/Plugins/Developer/Concert/ConcertSync/ConcertSyncClient/Source/ConcertSyncClient/Public/Replication/IConcertClientReplicationManager.h
@@ -12,6 +12,8 @@
 #include "Replication/Messages/Muting.h"
 #include "Replication/Messages/RestoreContent.h"
 
+#define UE_API CONCERTSYNCCLIENT_API
+
 template<typename ResultType>
 class TFuture;
 
@@ -52,7 +54,7 @@ namespace UE::ConcertSyncClient::Replication
  * Keeps a list of properties to send along with their send rules.
  * Tells the server which properties this client is interested in receiving.
  */
-class CONCERTSYNCCLIENT_API IConcertClientReplicationManager
+class IConcertClientReplicationManager
 {
 public:
 
@@ -82,9 +84,9 @@ public:
 	 */
 	virtual EStreamEnumerationResult ForEachRegisteredStream(TFunctionRef<EBreakBehavior(const FConcertReplicationStream& Stream)> Callback) const = 0;
 	/** @return Whether this manager has any registered streams (basically whether ForEachRegisteredStream returns EStreamEnumerationResult::Iterated). */
-	bool HasRegisteredStreams() const;
+	UE_API bool HasRegisteredStreams() const;
 	/** @return The streams registered with the server. */
-	TArray<FConcertReplicationStream> GetRegisteredStreams() const;
+	UE_API TArray<FConcertReplicationStream> GetRegisteredStreams() const;
 	
 	/**
 	 * Requests from the server to change the authority over some objects.
@@ -92,9 +94,9 @@ public:
 	 */
 	virtual TFuture<FConcertReplication_ChangeAuthority_Response> RequestAuthorityChange(FConcertReplication_ChangeAuthority_Request Args) = 0;
 	/** Util function that will request authority for all streams for the given objects. */
-	TFuture<FConcertReplication_ChangeAuthority_Response> TakeAuthorityOver(TConstArrayView<FSoftObjectPath> Objects);
+	UE_API TFuture<FConcertReplication_ChangeAuthority_Response> TakeAuthorityOver(TConstArrayView<FSoftObjectPath> Objects);
 	/** Util function that will let go over all authority of the given objects. */
-	TFuture<FConcertReplication_ChangeAuthority_Response> ReleaseAuthorityOf(TConstArrayView<FSoftObjectPath> Objects);
+	UE_API TFuture<FConcertReplication_ChangeAuthority_Response> ReleaseAuthorityOf(TConstArrayView<FSoftObjectPath> Objects);
 
 	enum class EAuthorityEnumerationResult { NoAuthorityAvailable, Iterated };
 	/**
@@ -107,7 +109,7 @@ public:
 	/** @return Whether this client has authority over ObjectPath for some stream. */
 	virtual bool HasAuthorityOver(const FSoftObjectPath& ObjectPath) const = 0;
 	/** @return All owned objects and associated owning streams. */
-	TMap<FSoftObjectPath, TSet<FGuid>> GetClientOwnedObjects() const;
+	UE_API TMap<FSoftObjectPath, TSet<FGuid>> GetClientOwnedObjects() const;
 
 	enum class ESyncControlEnumerationResult { NoneAvailable, Iterated };
 	/** Iterates through all objects the server told the client to replicate, which happens when this client has taken authority and another client wants to receive the data. */
@@ -117,7 +119,7 @@ public:
 	/** @return Whether this client has sync control for a specific object in a stream. */
 	virtual bool HasSyncControl(const FConcertObjectInStreamID& Object) const = 0;
 	/** Util for converting all sync controlled objects into a TSet. */
-	TSet<FConcertObjectInStreamID> GetSyncControlledObjects() const;
+	UE_API TSet<FConcertObjectInStreamID> GetSyncControlledObjects() const;
 
 	/**
 	 * Requests replication info about other clients, including the streams registered and which objects they have authority over (i.e. are sending).
@@ -137,12 +139,12 @@ public:
 	 */
 	virtual TFuture<FConcertReplication_ChangeMuteState_Response> ChangeMuteState(FConcertReplication_ChangeMuteState_Request Request) = 0;
 	/** Util function that will mute all of Objects. */
-	TFuture<FConcertReplication_ChangeMuteState_Response> MuteObjects(TConstArrayView<FSoftObjectPath> Objects, EConcertReplicationMuteOption Flags = EConcertReplicationMuteOption::ObjectAndSubobjects);
+	UE_API TFuture<FConcertReplication_ChangeMuteState_Response> MuteObjects(TConstArrayView<FSoftObjectPath> Objects, EConcertReplicationMuteOption Flags = EConcertReplicationMuteOption::ObjectAndSubobjects);
 	/** Util function that will unmute all of Objects. */
-	TFuture<FConcertReplication_ChangeMuteState_Response> UnmuteObjects(TSet<FSoftObjectPath> Objects, EConcertReplicationMuteOption Flags = EConcertReplicationMuteOption::ObjectAndSubobjects);
+	UE_API TFuture<FConcertReplication_ChangeMuteState_Response> UnmuteObjects(TSet<FSoftObjectPath> Objects, EConcertReplicationMuteOption Flags = EConcertReplicationMuteOption::ObjectAndSubobjects);
 	/** Gets the global mute state. */
 	virtual TFuture<FConcertReplication_QueryMuteState_Response> QueryMuteState(FConcertReplication_QueryMuteState_Request Request = {}) = 0;
-	TFuture<FConcertReplication_QueryMuteState_Response> QueryMuteState(TSet<FSoftObjectPath> Objects);
+	UE_API TFuture<FConcertReplication_QueryMuteState_Response> QueryMuteState(TSet<FSoftObjectPath> Objects);
 
 	/** Restore this client's stream content and authority to what a client had when they left. Only works if session has EConcertSyncSessionFlags::ShouldEnableReplicationActivities set. */
 	virtual TFuture<FConcertReplication_RestoreContent_Response> RestoreContent(FConcertReplication_RestoreContent_Request Request = {}) = 0;
@@ -189,4 +191,6 @@ public:
 	
 	virtual ~IConcertClientReplicationManager() = default;
 };
-	
\ No newline at end of file
+	
+
+#undef UE_API
diff --git a/Engine/Plugins/Developer/Concert/ConcertUI/ConcertSharedSlate/Source/ConcertSharedSlate/Public/ConcertFrontendStyle.h b/Engine/Plugins/Developer/Concert/ConcertUI/ConcertSharedSlate/Source/ConcertSharedSlate/Public/ConcertFrontendStyle.h
--- a/Engine/Plugins/Developer/Concert/ConcertUI/ConcertSharedSlate/Source/ConcertSharedSlate/Public/ConcertFrontendStyle.h
+++ b/Engine/Plugins/Developer/Concert/ConcertUI/ConcertSharedSlate/Source/ConcertSharedSlate/Public/ConcertFrontendStyle.h
@@ -6,19 +6,23 @@
 #include "Styling/ISlateStyle.h"
 #include "Styling/SlateStyle.h"
 
-class CONCERTSHAREDSLATE_API FConcertFrontendStyle
+#define UE_API CONCERTSHAREDSLATE_API
+
+class FConcertFrontendStyle
 {
 public:
 	
-	static void Initialize();
-	static void Shutdown();
+	static UE_API void Initialize();
+	static UE_API void Shutdown();
 
-	static TSharedPtr< class ISlateStyle > Get();
+	static UE_API TSharedPtr< class ISlateStyle > Get();
 
-	static FName GetStyleSetName();
+	static UE_API FName GetStyleSetName();
 	
 private:
 	
-	static FString InContent(const FString& RelativePath, const ANSICHAR* Extension);
-	static TSharedPtr< class FSlateStyleSet > StyleSet;
+	static UE_API FString InContent(const FString& RelativePath, const ANSICHAR* Extension);
+	static UE_API TSharedPtr< class FSlateStyleSet > StyleSet;
 };
+
+#undef UE_API
diff --git a/Engine/Plugins/Developer/Concert/ConcertUI/ConcertSharedSlate/Source/ConcertSharedSlate/Public/ConcertHeaderRowUtils.h b/Engine/Plugins/Developer/Concert/ConcertUI/ConcertSharedSlate/Source/ConcertSharedSlate/Public/ConcertHeaderRowUtils.h
--- a/Engine/Plugins/Developer/Concert/ConcertUI/ConcertSharedSlate/Source/ConcertSharedSlate/Public/ConcertHeaderRowUtils.h
+++ b/Engine/Plugins/Developer/Concert/ConcertUI/ConcertSharedSlate/Source/ConcertSharedSlate/Public/ConcertHeaderRowUtils.h
@@ -10,7 +10,7 @@ class SHeaderRow;
 class SWidget;
 
 USTRUCT()
-struct CONCERTSHAREDSLATE_API FColumnVisibilitySnapshot
+struct FColumnVisibilitySnapshot
 {
 	GENERATED_BODY()
 
diff --git a/Engine/Plugins/Developer/Concert/ConcertUI/ConcertSharedSlate/Source/ConcertSharedSlate/Public/Replication/Editor/Model/Extension/IStreamExtender.h b/Engine/Plugins/Developer/Concert/ConcertUI/ConcertSharedSlate/Source/ConcertSharedSlate/Public/Replication/Editor/Model/Extension/IStreamExtender.h
--- a/Engine/Plugins/Developer/Concert/ConcertUI/ConcertSharedSlate/Source/ConcertSharedSlate/Public/Replication/Editor/Model/Extension/IStreamExtender.h
+++ b/Engine/Plugins/Developer/Concert/ConcertUI/ConcertSharedSlate/Source/ConcertSharedSlate/Public/Replication/Editor/Model/Extension/IStreamExtender.h
@@ -12,7 +12,7 @@ namespace UE::ConcertSharedSlate
 	class IStreamExtensionContext;
 
 	/** When an object is added to IEditableReplicationStreamModel, these functions are called to allow adding additional objects and properties. */
-	class CONCERTSHAREDSLATE_API IStreamExtender
+	class IStreamExtender
 	{
 	public:
 
diff --git a/Engine/Plugins/Developer/Concert/ConcertUI/ConcertSharedSlate/Source/ConcertSharedSlate/Public/Replication/Editor/Model/Extension/IStreamExtensionContext.h b/Engine/Plugins/Developer/Concert/ConcertUI/ConcertSharedSlate/Source/ConcertSharedSlate/Public/Replication/Editor/Model/Extension/IStreamExtensionContext.h
--- a/Engine/Plugins/Developer/Concert/ConcertUI/ConcertSharedSlate/Source/ConcertSharedSlate/Public/Replication/Editor/Model/Extension/IStreamExtensionContext.h
+++ b/Engine/Plugins/Developer/Concert/ConcertUI/ConcertSharedSlate/Source/ConcertSharedSlate/Public/Replication/Editor/Model/Extension/IStreamExtensionContext.h
@@ -10,7 +10,7 @@ struct FConcertPropertyChain;
 namespace UE::ConcertSharedSlate
 {
 	/** Allows IStreamExtender to add default replication settings like properties and adding additional objects. */
-	class CONCERTSHAREDSLATE_API IStreamExtensionContext
+	class IStreamExtensionContext
 	{
 	public:
 
diff --git a/Engine/Plugins/Developer/Concert/ConcertUI/ConcertSharedSlate/Source/ConcertSharedSlate/Public/Replication/Editor/Model/IEditableMultiReplicationStreamModel.h b/Engine/Plugins/Developer/Concert/ConcertUI/ConcertSharedSlate/Source/ConcertSharedSlate/Public/Replication/Editor/Model/IEditableMultiReplicationStreamModel.h
--- a/Engine/Plugins/Developer/Concert/ConcertUI/ConcertSharedSlate/Source/ConcertSharedSlate/Public/Replication/Editor/Model/IEditableMultiReplicationStreamModel.h
+++ b/Engine/Plugins/Developer/Concert/ConcertUI/ConcertSharedSlate/Source/ConcertSharedSlate/Public/Replication/Editor/Model/IEditableMultiReplicationStreamModel.h
@@ -20,7 +20,7 @@ namespace UE::ConcertSharedSlate
 	 * - the set of readable streams includes the set of editable streams, but
 	 * - the number of editable streams can be less than or equal to the number of read-only streams.
 	 */
-	class CONCERTSHAREDSLATE_API IEditableMultiReplicationStreamModel : public IMultiReplicationStreamModel
+	class IEditableMultiReplicationStreamModel : public IMultiReplicationStreamModel
 	{
 	public:
 
@@ -55,4 +55,4 @@ namespace UE::ConcertSharedSlate
 			}
 		}
 	};
-}
\ No newline at end of file
+}
diff --git a/Engine/Plugins/Developer/Concert/ConcertUI/ConcertSharedSlate/Source/ConcertSharedSlate/Public/Replication/Editor/Model/IEditableReplicationStreamModel.h b/Engine/Plugins/Developer/Concert/ConcertUI/ConcertSharedSlate/Source/ConcertSharedSlate/Public/Replication/Editor/Model/IEditableReplicationStreamModel.h
--- a/Engine/Plugins/Developer/Concert/ConcertUI/ConcertSharedSlate/Source/ConcertSharedSlate/Public/Replication/Editor/Model/IEditableReplicationStreamModel.h
+++ b/Engine/Plugins/Developer/Concert/ConcertUI/ConcertSharedSlate/Source/ConcertSharedSlate/Public/Replication/Editor/Model/IEditableReplicationStreamModel.h
@@ -32,7 +32,7 @@ namespace UE::ConcertSharedSlate
 	 * Models may not always be writable. When editing a UAsset, it will be writable. However, if we join a multi-user
 	 * session we do not want to edit the objects nor properties - only read.
 	 */
-	class CONCERTSHAREDSLATE_API IEditableReplicationStreamModel : public IReplicationStreamModel
+	class IEditableReplicationStreamModel : public IReplicationStreamModel
 	{
 	public:
 
@@ -57,4 +57,4 @@ namespace UE::ConcertSharedSlate
 		/** Clear all objects in the model */
 		void Clear() { RemoveObjects(GetReplicatedObjects().Array()); }
 	};
-}
\ No newline at end of file
+}
diff --git a/Engine/Plugins/Developer/Concert/ConcertUI/ConcertSharedSlate/Source/ConcertSharedSlate/Public/Replication/Editor/Model/IMultiReplicationStreamModel.h b/Engine/Plugins/Developer/Concert/ConcertUI/ConcertSharedSlate/Source/ConcertSharedSlate/Public/Replication/Editor/Model/IMultiReplicationStreamModel.h
--- a/Engine/Plugins/Developer/Concert/ConcertUI/ConcertSharedSlate/Source/ConcertSharedSlate/Public/Replication/Editor/Model/IMultiReplicationStreamModel.h
+++ b/Engine/Plugins/Developer/Concert/ConcertUI/ConcertSharedSlate/Source/ConcertSharedSlate/Public/Replication/Editor/Model/IMultiReplicationStreamModel.h
@@ -10,7 +10,7 @@ namespace UE::ConcertSharedSlate
 	class IReplicationStreamModel;
 	
 	/** Model for UI displaying multiple streams at the same time. */
-	class CONCERTSHAREDSLATE_API IMultiReplicationStreamModel : public TSharedFromThis<IMultiReplicationStreamModel>
+	class IMultiReplicationStreamModel : public TSharedFromThis<IMultiReplicationStreamModel>
 	{
 	public:
 
@@ -27,4 +27,4 @@ namespace UE::ConcertSharedSlate
 
 		virtual ~IMultiReplicationStreamModel() = default;
 	};
-}
\ No newline at end of file
+}
diff --git a/Engine/Plugins/Developer/Concert/ConcertUI/ConcertSharedSlate/Source/ConcertSharedSlate/Public/Replication/Editor/Model/IReplicationStreamModel.h b/Engine/Plugins/Developer/Concert/ConcertUI/ConcertSharedSlate/Source/ConcertSharedSlate/Public/Replication/Editor/Model/IReplicationStreamModel.h
--- a/Engine/Plugins/Developer/Concert/ConcertUI/ConcertSharedSlate/Source/ConcertSharedSlate/Public/Replication/Editor/Model/IReplicationStreamModel.h
+++ b/Engine/Plugins/Developer/Concert/ConcertUI/ConcertSharedSlate/Source/ConcertSharedSlate/Public/Replication/Editor/Model/IReplicationStreamModel.h
@@ -17,7 +17,7 @@ namespace UE::ConcertSharedSlate
 	 * 
 	 * @see IEditableReplicationStreamModel 
 	 */
-	class CONCERTSHAREDSLATE_API IReplicationStreamModel : public TSharedFromThis<IReplicationStreamModel>
+	class IReplicationStreamModel : public TSharedFromThis<IReplicationStreamModel>
 	{
 	public:
 
diff --git a/Engine/Plugins/Developer/Concert/ConcertUI/ConcertSharedSlate/Source/ConcertSharedSlate/Public/Replication/Editor/Model/Object/IObjectHierarchyModel.h b/Engine/Plugins/Developer/Concert/ConcertUI/ConcertSharedSlate/Source/ConcertSharedSlate/Public/Replication/Editor/Model/Object/IObjectHierarchyModel.h
--- a/Engine/Plugins/Developer/Concert/ConcertUI/ConcertSharedSlate/Source/ConcertSharedSlate/Public/Replication/Editor/Model/Object/IObjectHierarchyModel.h
+++ b/Engine/Plugins/Developer/Concert/ConcertUI/ConcertSharedSlate/Source/ConcertSharedSlate/Public/Replication/Editor/Model/Object/IObjectHierarchyModel.h
@@ -9,6 +9,8 @@
 #include "UObject/SoftObjectPath.h"
 #include "UObject/SoftObjectPtr.h"
 
+#define UE_API CONCERTSHAREDSLATE_API
+
 enum class EBreakBehavior : uint8;
 
 namespace UE::ConcertSharedSlate
@@ -51,7 +53,7 @@ namespace UE::ConcertSharedSlate
 	 * In editor builds, an implementation could resolve object paths to AActors, USceneComponent, etc.
 	 * In programs, an implementation could be to compare all known object path strings.
 	 */
-	class CONCERTSHAREDSLATE_API IObjectHierarchyModel
+	class IObjectHierarchyModel
 	{
 	public:
 		
@@ -85,7 +87,7 @@ namespace UE::ConcertSharedSlate
 		virtual TOptional<FParentInfo> GetParentInfo(const TSoftObjectPtr<const UObject>& ChildObject) const = 0;
 
 		/** Util for iterating all subobjects. */
-		void ForEachChildRecursive(
+		UE_API void ForEachChildRecursive(
 			const TSoftObjectPtr<>& Root,
 			TFunctionRef<EBreakBehavior(const TSoftObjectPtr<>& Parent, const TSoftObjectPtr<>& ChildObject, EChildRelationship Relationship)> Callback,
 			EChildRelationshipFlags InclusionFlags = EChildRelationshipFlags::All
@@ -108,3 +110,5 @@ namespace UE::ConcertSharedSlate
 		virtual ~IObjectHierarchyModel() = default;
 	};
 }
+
+#undef UE_API
diff --git a/Engine/Plugins/Developer/Concert/ConcertUI/ConcertSharedSlate/Source/ConcertSharedSlate/Public/Replication/Editor/Model/Object/IObjectNameModel.h b/Engine/Plugins/Developer/Concert/ConcertUI/ConcertSharedSlate/Source/ConcertSharedSlate/Public/Replication/Editor/Model/Object/IObjectNameModel.h
--- a/Engine/Plugins/Developer/Concert/ConcertUI/ConcertSharedSlate/Source/ConcertSharedSlate/Public/Replication/Editor/Model/Object/IObjectNameModel.h
+++ b/Engine/Plugins/Developer/Concert/ConcertUI/ConcertSharedSlate/Source/ConcertSharedSlate/Public/Replication/Editor/Model/Object/IObjectNameModel.h
@@ -16,7 +16,7 @@ namespace UE::ConcertSharedSlate
 	 * In editor builds, this could use the USubobjectDataSubsystem to look up subobject names on an actor.
 	 * In programs, it could just default to what it says in the object path.
 	 */
-	class CONCERTSHAREDSLATE_API IObjectNameModel
+	class IObjectNameModel
 	{
 	public:
 
@@ -27,4 +27,4 @@ namespace UE::ConcertSharedSlate
 
 		virtual ~IObjectNameModel() = default;
 	};
-}
\ No newline at end of file
+}
diff --git a/Engine/Plugins/Developer/Concert/ConcertUI/ConcertSharedSlate/Source/ConcertSharedSlate/Public/Replication/Editor/View/Extension/CategoryRowGeneration.h b/Engine/Plugins/Developer/Concert/ConcertUI/ConcertSharedSlate/Source/ConcertSharedSlate/Public/Replication/Editor/View/Extension/CategoryRowGeneration.h
--- a/Engine/Plugins/Developer/Concert/ConcertUI/ConcertSharedSlate/Source/ConcertSharedSlate/Public/Replication/Editor/View/Extension/CategoryRowGeneration.h
+++ b/Engine/Plugins/Developer/Concert/ConcertUI/ConcertSharedSlate/Source/ConcertSharedSlate/Public/Replication/Editor/View/Extension/CategoryRowGeneration.h
@@ -1,4 +1,4 @@
-﻿// Copyright Epic Games, Inc. All Rights Reserved.
+// Copyright Epic Games, Inc. All Rights Reserved.
 
 #pragma once
 
@@ -30,7 +30,7 @@ namespace UE::ConcertSharedSlate
 	};
 	
 	/** A special reassignment view which shows the display object and all its children. */
-	class CONCERTSHAREDSLATE_API ICategoryRow
+	class ICategoryRow
 	{
 	public:
 
diff --git a/Engine/Plugins/Developer/Concert/ConcertUI/ConcertSharedSlate/Source/ConcertSharedSlate/Public/Replication/Editor/View/IMultiObjectPropertyAssignmentView.h b/Engine/Plugins/Developer/Concert/ConcertUI/ConcertSharedSlate/Source/ConcertSharedSlate/Public/Replication/Editor/View/IMultiObjectPropertyAssignmentView.h
--- a/Engine/Plugins/Developer/Concert/ConcertUI/ConcertSharedSlate/Source/ConcertSharedSlate/Public/Replication/Editor/View/IMultiObjectPropertyAssignmentView.h
+++ b/Engine/Plugins/Developer/Concert/ConcertUI/ConcertSharedSlate/Source/ConcertSharedSlate/Public/Replication/Editor/View/IMultiObjectPropertyAssignmentView.h
@@ -1,4 +1,4 @@
-﻿// Copyright Epic Games, Inc. All Rights Reserved.
+// Copyright Epic Games, Inc. All Rights Reserved.
 
 #pragma once
 
@@ -7,7 +7,7 @@
 namespace UE::ConcertSharedSlate
 {
 	/** A special reassignment view which shows the display object and all its children. */
-	class CONCERTSHAREDSLATE_API IMultiObjectPropertyAssignmentView : public IPropertyAssignmentView
+	class IMultiObjectPropertyAssignmentView : public IPropertyAssignmentView
 	{
 	public:
 
diff --git a/Engine/Plugins/Developer/Concert/ConcertUI/ConcertSharedSlate/Source/ConcertSharedSlate/Public/Replication/Editor/View/IMultiReplicationStreamEditor.h b/Engine/Plugins/Developer/Concert/ConcertUI/ConcertSharedSlate/Source/ConcertSharedSlate/Public/Replication/Editor/View/IMultiReplicationStreamEditor.h
--- a/Engine/Plugins/Developer/Concert/ConcertUI/ConcertSharedSlate/Source/ConcertSharedSlate/Public/Replication/Editor/View/IMultiReplicationStreamEditor.h
+++ b/Engine/Plugins/Developer/Concert/ConcertUI/ConcertSharedSlate/Source/ConcertSharedSlate/Public/Replication/Editor/View/IMultiReplicationStreamEditor.h
@@ -15,7 +15,7 @@ namespace UE::ConcertSharedSlate
 	 * Widget which edits multiple replication stream.
 	 * @see ReplicationWidgetFactories.h
 	 */
-	class CONCERTSHAREDSLATE_API IMultiReplicationStreamEditor : public SCompoundWidget
+	class IMultiReplicationStreamEditor : public SCompoundWidget
 	{
 	public:
 
diff --git a/Engine/Plugins/Developer/Concert/ConcertUI/ConcertSharedSlate/Source/ConcertSharedSlate/Public/Replication/Editor/View/IPropertyAssignmentView.h b/Engine/Plugins/Developer/Concert/ConcertUI/ConcertSharedSlate/Source/ConcertSharedSlate/Public/Replication/Editor/View/IPropertyAssignmentView.h
--- a/Engine/Plugins/Developer/Concert/ConcertUI/ConcertSharedSlate/Source/ConcertSharedSlate/Public/Replication/Editor/View/IPropertyAssignmentView.h
+++ b/Engine/Plugins/Developer/Concert/ConcertUI/ConcertSharedSlate/Source/ConcertSharedSlate/Public/Replication/Editor/View/IPropertyAssignmentView.h
@@ -40,7 +40,7 @@ namespace UE::ConcertSharedSlate
 	 *	- SMultiObjectAssignment, which displays the root object and its subobjects.
 	 * Generally, you can imagine the view as a tree view which has columns that can be injected (e.g. @see CreateBaseStreamEditor).
 	 */
-	class CONCERTSHAREDSLATE_API IPropertyAssignmentView
+	class IPropertyAssignmentView
 	{
 	public:
 
diff --git a/Engine/Plugins/Developer/Concert/ConcertUI/ConcertSharedSlate/Source/ConcertSharedSlate/Public/Replication/Editor/View/IPropertyTreeView.h b/Engine/Plugins/Developer/Concert/ConcertUI/ConcertSharedSlate/Source/ConcertSharedSlate/Public/Replication/Editor/View/IPropertyTreeView.h
--- a/Engine/Plugins/Developer/Concert/ConcertUI/ConcertSharedSlate/Source/ConcertSharedSlate/Public/Replication/Editor/View/IPropertyTreeView.h
+++ b/Engine/Plugins/Developer/Concert/ConcertUI/ConcertSharedSlate/Source/ConcertSharedSlate/Public/Replication/Editor/View/IPropertyTreeView.h
@@ -31,7 +31,7 @@ namespace UE::ConcertSharedSlate
 	};
 	
 	/** Represents a tree view displaying properties from a single class. */
-	class CONCERTSHAREDSLATE_API IPropertyTreeView
+	class IPropertyTreeView
 	{
 	public:
 		
diff --git a/Engine/Plugins/Developer/Concert/ConcertUI/ConcertSharedSlate/Source/ConcertSharedSlate/Public/Replication/Editor/View/IReplicationStreamEditor.h b/Engine/Plugins/Developer/Concert/ConcertUI/ConcertSharedSlate/Source/ConcertSharedSlate/Public/Replication/Editor/View/IReplicationStreamEditor.h
--- a/Engine/Plugins/Developer/Concert/ConcertUI/ConcertSharedSlate/Source/ConcertSharedSlate/Public/Replication/Editor/View/IReplicationStreamEditor.h
+++ b/Engine/Plugins/Developer/Concert/ConcertUI/ConcertSharedSlate/Source/ConcertSharedSlate/Public/Replication/Editor/View/IReplicationStreamEditor.h
@@ -11,6 +11,6 @@ namespace UE::ConcertSharedSlate
 	 * Widget which edits replication stream.
 	 * @see ReplicationWidgetFactories.h
 	 */
-	class CONCERTSHAREDSLATE_API IReplicationStreamEditor : public IReplicationStreamViewer
+	class IReplicationStreamEditor : public IReplicationStreamViewer
 	{};
-}
\ No newline at end of file
+}
diff --git a/Engine/Plugins/Developer/Concert/ConcertUI/ConcertSharedSlate/Source/ConcertSharedSlate/Public/Replication/Editor/View/IReplicationStreamViewer.h b/Engine/Plugins/Developer/Concert/ConcertUI/ConcertSharedSlate/Source/ConcertSharedSlate/Public/Replication/Editor/View/IReplicationStreamViewer.h
--- a/Engine/Plugins/Developer/Concert/ConcertUI/ConcertSharedSlate/Source/ConcertSharedSlate/Public/Replication/Editor/View/IReplicationStreamViewer.h
+++ b/Engine/Plugins/Developer/Concert/ConcertUI/ConcertSharedSlate/Source/ConcertSharedSlate/Public/Replication/Editor/View/IReplicationStreamViewer.h
@@ -5,13 +5,15 @@
 #include "Templates/SharedPointer.h"
 #include "Widgets/SCompoundWidget.h"
 
+#define UE_API CONCERTSHAREDSLATE_API
+
 namespace UE::ConcertSharedSlate
 {
 	/**
 	 * Widget which views a replication stream.
 	 * @see ReplicationWidgetFactories.h
 	 */
-	class CONCERTSHAREDSLATE_API IReplicationStreamViewer : public SCompoundWidget
+	class IReplicationStreamViewer : public SCompoundWidget
 	{
 	public:
 
@@ -31,7 +33,7 @@ namespace UE::ConcertSharedSlate
 
 		/** @return The objects for which the properties are being edited / displayed. */
 		UE_DEPRECATED(5.5, "Use GetSelectedObjects instead.")
-		virtual TArray<FSoftObjectPath> GetObjectsBeingPropertyEdited() const;
+		UE_API virtual TArray<FSoftObjectPath> GetObjectsBeingPropertyEdited() const;
 
 		/** Selects the objects in the outliner and deselects all others. Only actors can be selected. Invalid entries are filtered out. */
 		virtual void SetSelectedObjects(TConstArrayView<TSoftObjectPtr<>> Objects) = 0;
@@ -41,3 +43,5 @@ namespace UE::ConcertSharedSlate
 		virtual ~IReplicationStreamViewer() = default;
 	};
 }
+
+#undef UE_API
diff --git a/Engine/Plugins/Developer/Concert/ConcertUI/ConcertSharedSlate/Source/ConcertSharedSlate/Public/SConcertScrollBox.h b/Engine/Plugins/Developer/Concert/ConcertUI/ConcertSharedSlate/Source/ConcertSharedSlate/Public/SConcertScrollBox.h
--- a/Engine/Plugins/Developer/Concert/ConcertUI/ConcertSharedSlate/Source/ConcertSharedSlate/Public/SConcertScrollBox.h
+++ b/Engine/Plugins/Developer/Concert/ConcertUI/ConcertSharedSlate/Source/ConcertSharedSlate/Public/SConcertScrollBox.h
@@ -7,12 +7,14 @@
 #include "Widgets/SCompoundWidget.h"
 #include "Framework/SlateDelegates.h"
 
+#define UE_API CONCERTSHAREDSLATE_API
+
 class SScrollBox;
 class SButton;
 class SScrollBar;
 
 
-class CONCERTSHAREDSLATE_API SConcertScrollBox : public SCompoundWidget
+class SConcertScrollBox : public SCompoundWidget
 {
 public:
 
@@ -35,36 +37,36 @@ public:
 	* @param InArgs The Slate argument list.
 	* @param ConstructUnderMajorTab The major tab which will contain the Session History widget.
 	*/
-	void Construct(const FArguments& InArgs);
+	UE_API void Construct(const FArguments& InArgs);
 
 	/** SWidget interface. */
-	virtual void Tick(const FGeometry& AllottedGeometry, const double InCurrentTime, const float InDeltaTime) override;
+	UE_API virtual void Tick(const FGeometry& AllottedGeometry, const double InCurrentTime, const float InDeltaTime) override;
 	/** End of SWidget interface. */
 
 	/**
 	 * @return a new slot. Slots contain children for SScrollBox 
 	 */
-	static FSlot::FSlotArguments Slot();
+	static UE_API FSlot::FSlotArguments Slot();
 
 private:
 	
 	/** Delegate called to handle a user scrolling. */
-	void HandleUserScrolled(float);
+	UE_API void HandleUserScrolled(float);
 
 	/** Checks if the scrollbar is currently at the bottom of its track. */
-	bool IsAtBottom();
+	UE_API bool IsAtBottom();
 
 	/** Callback for scrolling to the start of the list. */
-	FReply HandleScrollToStart();
+	UE_API FReply HandleScrollToStart();
 
 	/** Callback for scrolling to the end of the list. */
-	FReply HandleScrollToEnd();
+	UE_API FReply HandleScrollToEnd();
 
 	/** Callback for getting the ScrollBar buttons visibility. */
-	EVisibility HandleScrollButtonsVisibility() const;
+	UE_API EVisibility HandleScrollButtonsVisibility() const;
 
 	/** Creates a button with the FontAwesome font. */
-	TSharedRef<SButton> CreateScrollBarButton(const FText& InToolTip, const FString& InIcon, FOnClicked OnClickedDelegate) const;
+	UE_API TSharedRef<SButton> CreateScrollBarButton(const FText& InToolTip, const FString& InIcon, FOnClicked OnClickedDelegate) const;
 
 private:
 	
@@ -80,3 +82,5 @@ private:
 	/** Holds the inner ScrollBox. */
 	TSharedPtr<SScrollBox> ScrollBox;
 };
+
+#undef UE_API
diff --git a/Engine/Plugins/Developer/Concert/ConcertUI/ConcertSharedSlate/Source/ConcertSharedSlate/Public/SConcertSessionRecovery.h b/Engine/Plugins/Developer/Concert/ConcertUI/ConcertSharedSlate/Source/ConcertSharedSlate/Public/SConcertSessionRecovery.h
--- a/Engine/Plugins/Developer/Concert/ConcertUI/ConcertSharedSlate/Source/ConcertSharedSlate/Public/SConcertSessionRecovery.h
+++ b/Engine/Plugins/Developer/Concert/ConcertUI/ConcertSharedSlate/Source/ConcertSharedSlate/Public/SConcertSessionRecovery.h
@@ -8,6 +8,8 @@
 #include "Widgets/DeclarativeSyntaxSupport.h"
 #include "Widgets/SCompoundWidget.h"
 
+#define UE_API CONCERTSHAREDSLATE_API
+
 struct FConcertSessionActivity;
 struct FConcertClientInfo;
 class FConcertSessionActivitiesOptions;
@@ -15,7 +17,7 @@ class FConcertSessionActivitiesOptions;
 /**
  * Displays the list of activities available for recovery and let the user select what should or shouldn't be recovered.
  */
-class CONCERTSHAREDSLATE_API SConcertSessionRecovery : public SCompoundWidget
+class SConcertSessionRecovery : public SCompoundWidget
 {
 public:
 	SLATE_BEGIN_ARGS(SConcertSessionRecovery)
@@ -98,41 +100,41 @@ public:
 	 * Construct the recovery widget.
 	 * @param InArgs The widgets arguments and attributes.
 	 */
-	void Construct(const FArguments& InArgs);
+	UE_API void Construct(const FArguments& InArgs);
 
 	/** Returns the activity, selected by the user, through which the session should be (or was) recovered or null to prevent recovery. */
 	TSharedPtr<FConcertSessionActivity> GetRecoverThroughItem() { return RecoveryThroughItem; }
 
 	/** Removes all activities from the view and reset the activity stream. */
-	void Reset();
+	UE_API void Reset();
 
 	/** Returns the total number of activities currently stored (no filter applied). */
-	int32 GetTotalActivityNum() const;
+	UE_API int32 GetTotalActivityNum() const;
 
 	/** Returns the tooltip shown for the 'Recovery All' button. */
-	FText GetRecoverAllButtonTooltip() const;
+	UE_API FText GetRecoverAllButtonTooltip() const;
 
 	/** Returns the most recent activity available, ignoring the current filter. */
-	TSharedPtr<FConcertSessionActivity> GetMostRecentActivity() const;
+	UE_API TSharedPtr<FConcertSessionActivity> GetMostRecentActivity() const;
 
 private:
 
-	TArray<FActivityColumn> MakeColumns(const FArguments& InArgs) const;
+	UE_API TArray<FActivityColumn> MakeColumns(const FArguments& InArgs) const;
 	
-	void OnSearchTextChanged(const FText& InFilterText);
-	void OnSearchTextCommitted(const FText& InFilterText, ETextCommit::Type CommitType);
-	FText HighlightSearchText() const;
+	UE_API void OnSearchTextChanged(const FText& InFilterText);
+	UE_API void OnSearchTextCommitted(const FText& InFilterText, ETextCommit::Type CommitType);
+	UE_API FText HighlightSearchText() const;
 
-	FReply OnCancelRecoveryClicked();
-	FReply OnRecoverAllClicked();
+	UE_API FReply OnCancelRecoveryClicked();
+	UE_API FReply OnRecoverAllClicked();
 
-	EVisibility GetRecoverThroughButtonVisibility(TSharedPtr<FConcertSessionActivity> Activity);
-	FText GetRecoverThroughButtonTooltip() const;
-	TSharedPtr<SWidget> MakeRecoverThroughWidget(TWeakPtr<FConcertSessionActivity>, const FName&);
-	void RecoverThrough(TSharedPtr<FConcertSessionActivity> Item);
+	UE_API EVisibility GetRecoverThroughButtonVisibility(TSharedPtr<FConcertSessionActivity> Activity);
+	UE_API FText GetRecoverThroughButtonTooltip() const;
+	UE_API TSharedPtr<SWidget> MakeRecoverThroughWidget(TWeakPtr<FConcertSessionActivity>, const FName&);
+	UE_API void RecoverThrough(TSharedPtr<FConcertSessionActivity> Item);
 
 	/** Close the windows hosting this recovery widget. */
-	void DismissWindow();
+	UE_API void DismissWindow();
 
 private:
 	/** Display the session activities. */
@@ -165,3 +167,5 @@ private:
 	/** The function invoked when the user clicks the cancel button. Might not be bound. */
 	TFunction<void()> OnCancelFn;
 };
+
+#undef UE_API
diff --git a/Engine/Plugins/Developer/Concert/ConcertUI/ConcertSharedSlate/Source/ConcertSharedSlate/Public/SPackageDetails.h b/Engine/Plugins/Developer/Concert/ConcertUI/ConcertSharedSlate/Source/ConcertSharedSlate/Public/SPackageDetails.h
--- a/Engine/Plugins/Developer/Concert/ConcertUI/ConcertSharedSlate/Source/ConcertSharedSlate/Public/SPackageDetails.h
+++ b/Engine/Plugins/Developer/Concert/ConcertUI/ConcertSharedSlate/Source/ConcertSharedSlate/Public/SPackageDetails.h
@@ -7,9 +7,11 @@
 #include "Widgets/SCompoundWidget.h"
 #include "Widgets/Views/SListView.h"
 
+#define UE_API CONCERTSHAREDSLATE_API
+
 struct FConcertPackageInfo;
 
-class CONCERTSHAREDSLATE_API SPackageDetails : public SCompoundWidget
+class SPackageDetails : public SCompoundWidget
 {
 public:
 	SLATE_BEGIN_ARGS(SPackageDetails)
@@ -21,7 +23,7 @@ public:
 	*
 	* @param InArgs The Slate argument list.
 	*/
-	void Construct(const FArguments& InArgs);
+	UE_API void Construct(const FArguments& InArgs);
 
 	/**
 	 * Set the current package to have its details displayed.
@@ -29,7 +31,7 @@ public:
 	 * @param InRevision The target package's revision.
 	 * @param InModifiedBy The username of the user who modified the package.
 	 */
-	void SetPackageInfo(const FConcertPackageInfo& InPackageInfo, const int64 InRevision, const FString& InModifiedBy);
+	UE_API void SetPackageInfo(const FConcertPackageInfo& InPackageInfo, const int64 InRevision, const FString& InModifiedBy);
 
 private:
 
@@ -45,7 +47,7 @@ private:
 	using FListItemPtr = TSharedPtr<FPackageDetailsRow>;
 
 	/** Callback for generating package info rows. */
-	TSharedRef<ITableRow> HandleGenerateRow(FListItemPtr InItem, const TSharedRef<STableViewBase>& OwnerTable) const;
+	UE_API TSharedRef<ITableRow> HandleGenerateRow(FListItemPtr InItem, const TSharedRef<STableViewBase>& OwnerTable) const;
 
 private:
 
@@ -54,4 +56,6 @@ private:
 
 	/** Holds the package details rows. */
 	TArray<FListItemPtr> DetailItems;
-};
\ No newline at end of file
+};
+
+#undef UE_API
diff --git a/Engine/Plugins/Developer/Concert/ConcertUI/ConcertSharedSlate/Source/ConcertSharedSlate/Public/Session/Activity/ActivityColumn.h b/Engine/Plugins/Developer/Concert/ConcertUI/ConcertSharedSlate/Source/ConcertSharedSlate/Public/Session/Activity/ActivityColumn.h
--- a/Engine/Plugins/Developer/Concert/ConcertUI/ConcertSharedSlate/Source/ConcertSharedSlate/Public/Session/Activity/ActivityColumn.h
+++ b/Engine/Plugins/Developer/Concert/ConcertUI/ConcertSharedSlate/Source/ConcertSharedSlate/Public/Session/Activity/ActivityColumn.h
@@ -1,4 +1,4 @@
-﻿// Copyright Epic Games, Inc. All Rights Reserved.
+// Copyright Epic Games, Inc. All Rights Reserved.
 
 #pragma once
 
@@ -6,11 +6,13 @@
 #include "Widgets/SOverlay.h"
 #include "Widgets/Views/SHeaderRow.h"
 
+#define UE_API CONCERTSHAREDSLATE_API
+
 class SConcertSessionActivities;
 struct FConcertSessionActivity;
 
 /** A column in SConcertSessionActivities */
-class CONCERTSHAREDSLATE_API FActivityColumn final : public SHeaderRow::FColumn
+class FActivityColumn final : public SHeaderRow::FColumn
 {
 public:
 
@@ -40,8 +42,8 @@ public:
 	/** Determines whether this column is the first, etc. */
 	FActivityColumn& ColumnSortOrder(int32 NewValue) { ColumnSortOrderValue = NewValue; return *this; }
 	
-	void BuildColumnWidget(const TSharedRef<SConcertSessionActivities>& Owner, const TSharedRef<FConcertSessionActivity>& Activity, SOverlay::FScopedWidgetSlotArguments& Slot) const;
-	void ExecutePopulateSearchString(const TSharedRef<SConcertSessionActivities>& Owner, const FConcertSessionActivity& Activity, TArray<FString>& ExistingSearchStrings) const;
+	UE_API void BuildColumnWidget(const TSharedRef<SConcertSessionActivities>& Owner, const TSharedRef<FConcertSessionActivity>& Activity, SOverlay::FScopedWidgetSlotArguments& Slot) const;
+	UE_API void ExecutePopulateSearchString(const TSharedRef<SConcertSessionActivities>& Owner, const FConcertSessionActivity& Activity, TArray<FString>& ExistingSearchStrings) const;
 	int32 GetColumnSortOrderValue() const { return ColumnSortOrderValue; }
 	
 private:
@@ -52,3 +54,5 @@ private:
 	/** Determines whether this column is the first, etc. */
 	int32 ColumnSortOrderValue = 0;
 };
+
+#undef UE_API
diff --git a/Engine/Plugins/Developer/Concert/ConcertUI/ConcertSharedSlate/Source/ConcertSharedSlate/Public/Session/Activity/IConcertReflectionDataProvider.h b/Engine/Plugins/Developer/Concert/ConcertUI/ConcertSharedSlate/Source/ConcertSharedSlate/Public/Session/Activity/IConcertReflectionDataProvider.h
--- a/Engine/Plugins/Developer/Concert/ConcertUI/ConcertSharedSlate/Source/ConcertSharedSlate/Public/Session/Activity/IConcertReflectionDataProvider.h
+++ b/Engine/Plugins/Developer/Concert/ConcertUI/ConcertSharedSlate/Source/ConcertSharedSlate/Public/Session/Activity/IConcertReflectionDataProvider.h
@@ -10,7 +10,7 @@ struct FConcertTransactionEventBase;
 namespace UE::ConcertSharedSlate
 {
 	/** Passed on to SUndoHistoryDetails to display reflection data in non-editor builds. Used e.g. by the server. */
-	class CONCERTSHAREDSLATE_API IConcertReflectionDataProvider : public UndoHistory::IReflectionDataProvider 
+	class IConcertReflectionDataProvider : public UndoHistory::IReflectionDataProvider 
 	{
 	public:
 
diff --git a/Engine/Plugins/Developer/Concert/ConcertUI/ConcertSharedSlate/Source/ConcertSharedSlate/Public/Session/Activity/SConcertSessionActivities.h b/Engine/Plugins/Developer/Concert/ConcertUI/ConcertSharedSlate/Source/ConcertSharedSlate/Public/Session/Activity/SConcertSessionActivities.h
--- a/Engine/Plugins/Developer/Concert/ConcertUI/ConcertSharedSlate/Source/ConcertSharedSlate/Public/Session/Activity/SConcertSessionActivities.h
+++ b/Engine/Plugins/Developer/Concert/ConcertUI/ConcertSharedSlate/Source/ConcertSharedSlate/Public/Session/Activity/SConcertSessionActivities.h
@@ -16,6 +16,8 @@
 #include "Widgets/Views/SListView.h"
 #include "Widgets/Layout/SExpandableArea.h"
 
+#define UE_API CONCERTSHAREDSLATE_API
+
 struct FConcertSessionActivity;
 struct FConcertClientInfo;
 struct FConcertPackageInfo;
@@ -45,7 +47,7 @@ enum class EConcertActivityFilterFlags
 ENUM_CLASS_FLAGS(EConcertActivityFilterFlags);
 
 /** Displays session activities in a table view. */
-class CONCERTSHAREDSLATE_API SConcertSessionActivities : public SCompoundWidget
+class SConcertSessionActivities : public SCompoundWidget
 {
 public:
 	
@@ -170,15 +172,15 @@ public:
 	 * Construct the recovery widget.
 	 * @param InArgs The widgets arguments and attributes.
 	 */
-	void Construct(const FArguments& InArgs);
+	UE_API void Construct(const FArguments& InArgs);
 
 	/** Requests and consumes more activities if needed and/or possible. */
-	virtual void Tick(const FGeometry& AllottedGeometry, const double InCurrentTime, const float InDeltaTime) override;
+	UE_API virtual void Tick(const FGeometry& AllottedGeometry, const double InCurrentTime, const float InDeltaTime) override;
 
 	/** Returns the activity selected or null if none is selected. */
-	TSharedPtr<FConcertSessionActivity> GetSelectedActivity() const;
+	UE_API TSharedPtr<FConcertSessionActivity> GetSelectedActivity() const;
 	/** Returns all selected activities */
-	TArray<TSharedPtr<FConcertSessionActivity>> GetSelectedActivities() const;
+	UE_API TArray<TSharedPtr<FConcertSessionActivity>> GetSelectedActivities() const;
 
 	/** Returns the total number of activities currently stored (no filter applied). */
 	int32 GetTotalActivityNum() const { return AllActivities.Num(); }
@@ -190,24 +192,24 @@ public:
 	int32 GetIgnoredActivityNum() const { return IgnoredActivityNum; }
 
 	/** Returns the most recent activity available, ignoring the current filter. */
-	TSharedPtr<FConcertSessionActivity> GetMostRecentActivity() const;
+	UE_API TSharedPtr<FConcertSessionActivity> GetMostRecentActivity() const;
 
 	/** Returns true if the column names is the last one (most right one). */
-	bool IsLastColumn(const FName& ColumnId) const;
+	UE_API bool IsLastColumn(const FName& ColumnId) const;
 
 	/** Clears all activities displayed. */
-	void ResetActivityList();
+	UE_API void ResetActivityList();
 
 	/** Append an activity to the view. Used to populate the view from a live session. */
-	void Append(TSharedPtr<FConcertSessionActivity> Activity);
+	UE_API void Append(TSharedPtr<FConcertSessionActivity> Activity);
 
 	/** Request the view to refresh. */
-	void RequestRefresh();
+	UE_API void RequestRefresh();
 
 	/** Asks the view to update the text search filter. */
-	FText UpdateTextFilter(const FText& InFilterText);
+	UE_API FText UpdateTextFilter(const FText& InFilterText);
 	
-	void OnColumnVisibilitySettingsChanged(const FColumnVisibilitySnapshot& ColumnSnapshot);
+	UE_API void OnColumnVisibilitySettingsChanged(const FColumnVisibilitySnapshot& ColumnSnapshot);
 	
 	const TAttribute<FText>& GetHighlightText() const { return HighlightText; }
 	const TAttribute<ETimeFormat>& GetTimeFormat() const { return TimeFormat; }
@@ -224,38 +226,38 @@ public:
 private:
 
 	// Widget creation
-	TSharedRef<SHeaderRow> CreateHeaderRow(const FArguments& InArgs);
-	TSharedRef<ITableRow> OnGenerateActivityRowWidget(TSharedPtr<FConcertSessionActivity> Item, const TSharedRef<STableViewBase>& OwnerTable);
-	void FetchActivities();
+	UE_API TSharedRef<SHeaderRow> CreateHeaderRow(const FArguments& InArgs);
+	UE_API TSharedRef<ITableRow> OnGenerateActivityRowWidget(TSharedPtr<FConcertSessionActivity> Item, const TSharedRef<STableViewBase>& OwnerTable);
+	UE_API void FetchActivities();
 
 	// Filtering
-	EConcertActivityFilterFlags QueryActiveActivityFilters() const;
-	void OnActivityFilterUpdated();
-	bool PassesFilters(const FConcertSessionActivity& Activity);
+	UE_API EConcertActivityFilterFlags QueryActiveActivityFilters() const;
+	UE_API void OnActivityFilterUpdated();
+	UE_API bool PassesFilters(const FConcertSessionActivity& Activity);
 
 	// List view
-	void OnListViewScrolled(double InScrollOffset);
-	void OnListViewSelectionChanged(TSharedPtr<FConcertSessionActivity> InActivity, ESelectInfo::Type SelectInfo);
+	UE_API void OnListViewScrolled(double InScrollOffset);
+	UE_API void OnListViewSelectionChanged(TSharedPtr<FConcertSessionActivity> InActivity, ESelectInfo::Type SelectInfo);
 
 	// Search
-	void PopulateSearchStrings(const FConcertSessionActivity& Item, TArray<FString>& OutSearchStrings);
+	UE_API void PopulateSearchStrings(const FConcertSessionActivity& Item, TArray<FString>& OutSearchStrings);
 
 	// Details
-	void OnDetailsAreaExpansionChanged(bool bExpanded);
-	void UpdateDetailArea(TSharedPtr<FConcertSessionActivity> InSelectedActivity);
-	void UpdateDetailsByEventPayload(TSharedPtr<FConcertSessionActivity> InSelectedActivity);
-	void UpdateDetailsByRequest(TSharedPtr<FConcertSessionActivity> InSelectedActivity);
+	UE_API void OnDetailsAreaExpansionChanged(bool bExpanded);
+	UE_API void UpdateDetailArea(TSharedPtr<FConcertSessionActivity> InSelectedActivity);
+	UE_API void UpdateDetailsByEventPayload(TSharedPtr<FConcertSessionActivity> InSelectedActivity);
+	UE_API void UpdateDetailsByRequest(TSharedPtr<FConcertSessionActivity> InSelectedActivity);
 	
-	void SetDetailsPanelVisibility(const SWidget* VisiblePanel);
-	void DisplayTransactionDetails(const FConcertSessionActivity& Activity, const FConcertTransactionEventBase& InTransaction);
-	void DisplayPackageDetails(const FConcertSessionActivity& Activity, int64 PackageRevision, const FConcertPackageInfo& PackageInfo);
+	UE_API void SetDetailsPanelVisibility(const SWidget* VisiblePanel);
+	UE_API void DisplayTransactionDetails(const FConcertSessionActivity& Activity, const FConcertTransactionEventBase& InTransaction);
+	UE_API void DisplayPackageDetails(const FConcertSessionActivity& Activity, int64 PackageRevision, const FConcertPackageInfo& PackageInfo);
 	
-	FText GetNoDetailsText() const;
+	UE_API FText GetNoDetailsText() const;
 	EVisibility GetDetailAreaVisibility() const { return DetailsAreaVisibility; }
 	SSplitter::ESizeRule GetDetailsAreaSizeRule() const { return bDetailsAreaExpanded ? SSplitter::ESizeRule::FractionOfParent : SSplitter::ESizeRule::SizeToContent; }
 	
 	/** Invoked after all widget ticks get processed. Used to check if the user scrolled or not. */
-	void OnPostTick(float);
+	UE_API void OnPostTick(float);
 
 private:
 	
@@ -371,19 +373,19 @@ private:
  * Manages the various options exposed by SConcertSessionActivities widget such as enabling/disabling filtering,
  * changing the time format, etc.
  */
-class CONCERTSHAREDSLATE_API FConcertSessionActivitiesOptions : public TSharedFromThis<FConcertSessionActivitiesOptions>
+class FConcertSessionActivitiesOptions : public TSharedFromThis<FConcertSessionActivitiesOptions>
 {
 public:
 	
 	DECLARE_DELEGATE_OneParam(FExtendContextMenu, FMenuBuilder&)
 
 	/** Makes a status bar widget displaying the activity shown and the standard view options button. */
-	TSharedRef<SWidget> MakeStatusBar(TAttribute<int32> Total, TAttribute<int32> Displayed);
+	UE_API TSharedRef<SWidget> MakeStatusBar(TAttribute<int32> Total, TAttribute<int32> Displayed);
 	
 	/** Makes a standard View Options widget, displaying the eye ball icon and showing the possible options. */
-	TSharedRef<SWidget> MakeViewOptionsComboButton(TOptional<FExtendContextMenu> ExtendMenu = {});
+	UE_API TSharedRef<SWidget> MakeViewOptionsComboButton(TOptional<FExtendContextMenu> ExtendMenu = {});
 	/** Returns a Menu widgets containing the available options. */
-	TSharedRef<SWidget> MakeMenuWidget(TOptional<FExtendContextMenu> ExtendMenu = {});
+	UE_API TSharedRef<SWidget> MakeMenuWidget(TOptional<FExtendContextMenu> ExtendMenu = {});
 
 	/** Returns the time format option. */
 	ETimeFormat GetTimeFormat() const { return bDisplayRelativeTime ? ETimeFormat::Relative : ETimeFormat::Absolute; }
@@ -401,7 +403,7 @@ public:
 	EVisibility GetIgnoredActivitiesVisibility() const { return bDisplayIgnoredActivities ? EVisibility::Visible : EVisibility::Hidden; }
 
 	/** Invoked when an options is togged from the displayed menu widget. */
-	void OnOptionToggled(const FName CheckBoxId);
+	UE_API void OnOptionToggled(const FName CheckBoxId);
 
 	/** Enables the 'connection activity' filter check box. (Show Connection Activities). */
 	bool bEnableConnectionActivityFiltering = true;
@@ -430,3 +432,5 @@ public:
 	/** If ignored activity filtering is enabled, controls whether ignored activities are filtered out.*/
 	bool bDisplayIgnoredActivities = false;
 };
+
+#undef UE_API
diff --git a/Engine/Plugins/Developer/Concert/ConcertUI/ConcertSharedSlate/Source/ConcertSharedSlate/Public/Session/Browser/ConcertSessionBrowserSettings.h b/Engine/Plugins/Developer/Concert/ConcertUI/ConcertSharedSlate/Source/ConcertSharedSlate/Public/Session/Browser/ConcertSessionBrowserSettings.h
--- a/Engine/Plugins/Developer/Concert/ConcertUI/ConcertSharedSlate/Source/ConcertSharedSlate/Public/Session/Browser/ConcertSessionBrowserSettings.h
+++ b/Engine/Plugins/Developer/Concert/ConcertUI/ConcertSharedSlate/Source/ConcertSharedSlate/Public/Session/Browser/ConcertSessionBrowserSettings.h
@@ -8,8 +8,8 @@
 #include "ConcertSessionBrowserSettings.generated.h"
 
 /** Serializes the multi-user session browser settings like the active filters. */
-UCLASS(config=Editor)
-class CONCERTSHAREDSLATE_API UConcertSessionBrowserSettings : public UObject
+UCLASS(MinimalAPI, config=Editor)
+class UConcertSessionBrowserSettings : public UObject
 {
 	GENERATED_BODY()
 
diff --git a/Engine/Plugins/Developer/Concert/ConcertUI/ConcertSharedSlate/Source/ConcertSharedSlate/Public/Session/Browser/IConcertSessionBrowserController.h b/Engine/Plugins/Developer/Concert/ConcertUI/ConcertSharedSlate/Source/ConcertSharedSlate/Public/Session/Browser/IConcertSessionBrowserController.h
--- a/Engine/Plugins/Developer/Concert/ConcertUI/ConcertSharedSlate/Source/ConcertSharedSlate/Public/Session/Browser/IConcertSessionBrowserController.h
+++ b/Engine/Plugins/Developer/Concert/ConcertUI/ConcertSharedSlate/Source/ConcertSharedSlate/Public/Session/Browser/IConcertSessionBrowserController.h
@@ -10,7 +10,7 @@
  * is to keep the UI code as decoupled as possible from the API used to query it. It encapsulate the asynchronous code and provide a
  * simpler API to the UI.
  */
-class CONCERTSHAREDSLATE_API IConcertSessionBrowserController
+class IConcertSessionBrowserController
 {
 public:
 
diff --git a/Engine/Plugins/Developer/Concert/ConcertUI/ConcertSharedSlate/Source/ConcertSharedSlate/Public/Session/Browser/Items/ConcertActiveGroupTreeItem.h b/Engine/Plugins/Developer/Concert/ConcertUI/ConcertSharedSlate/Source/ConcertSharedSlate/Public/Session/Browser/Items/ConcertActiveGroupTreeItem.h
--- a/Engine/Plugins/Developer/Concert/ConcertUI/ConcertSharedSlate/Source/ConcertSharedSlate/Public/Session/Browser/Items/ConcertActiveGroupTreeItem.h
+++ b/Engine/Plugins/Developer/Concert/ConcertUI/ConcertSharedSlate/Source/ConcertSharedSlate/Public/Session/Browser/Items/ConcertActiveGroupTreeItem.h
@@ -7,7 +7,7 @@
 #include "ConcertSessionTreeItem.h"
 
 /** Groups several active sessions */
-class CONCERTSHAREDSLATE_API FConcertActiveGroupTreeItem : public FConcertGroupTreeItem
+class FConcertActiveGroupTreeItem : public FConcertGroupTreeItem
 {
 	using Super = FConcertGroupTreeItem;
 public:
@@ -15,4 +15,4 @@ public:
 	FConcertActiveGroupTreeItem(FGetSessions GetSessionsFunc)
 		: Super(MoveTemp(GetSessionsFunc), { FConcertSessionTreeItem::EType::ActiveSession, FConcertSessionTreeItem::EType::NewSession, FConcertSessionTreeItem::EType::SaveSession })
 	{}
-};
\ No newline at end of file
+};
diff --git a/Engine/Plugins/Developer/Concert/ConcertUI/ConcertSharedSlate/Source/ConcertSharedSlate/Public/Session/Browser/Items/ConcertArchivedGroupTreeItem.h b/Engine/Plugins/Developer/Concert/ConcertUI/ConcertSharedSlate/Source/ConcertSharedSlate/Public/Session/Browser/Items/ConcertArchivedGroupTreeItem.h
--- a/Engine/Plugins/Developer/Concert/ConcertUI/ConcertSharedSlate/Source/ConcertSharedSlate/Public/Session/Browser/Items/ConcertArchivedGroupTreeItem.h
+++ b/Engine/Plugins/Developer/Concert/ConcertUI/ConcertSharedSlate/Source/ConcertSharedSlate/Public/Session/Browser/Items/ConcertArchivedGroupTreeItem.h
@@ -6,7 +6,7 @@
 #include "ConcertGroupTreeItem.h"
 
 /** Groups archived sessions */
-class CONCERTSHAREDSLATE_API FConcertArchivedGroupTreeItem : public FConcertGroupTreeItem
+class FConcertArchivedGroupTreeItem : public FConcertGroupTreeItem
 {
 	using Super = FConcertGroupTreeItem;
 public:
@@ -14,4 +14,4 @@ public:
 	FConcertArchivedGroupTreeItem(FGetSessions GetSessionsFunc)
 		: Super(MoveTemp(GetSessionsFunc), { FConcertSessionTreeItem::EType::ArchivedSession, FConcertSessionTreeItem::EType::RestoreSession })
 	{}
-};
\ No newline at end of file
+};
diff --git a/Engine/Plugins/Developer/Concert/ConcertUI/ConcertSharedSlate/Source/ConcertSharedSlate/Public/Session/Browser/Items/ConcertGroupTreeItem.h b/Engine/Plugins/Developer/Concert/ConcertUI/ConcertSharedSlate/Source/ConcertSharedSlate/Public/Session/Browser/Items/ConcertGroupTreeItem.h
--- a/Engine/Plugins/Developer/Concert/ConcertUI/ConcertSharedSlate/Source/ConcertSharedSlate/Public/Session/Browser/Items/ConcertGroupTreeItem.h
+++ b/Engine/Plugins/Developer/Concert/ConcertUI/ConcertSharedSlate/Source/ConcertSharedSlate/Public/Session/Browser/Items/ConcertGroupTreeItem.h
@@ -7,7 +7,7 @@
 #include "ConcertTreeItem.h"
 
 /** Groups several FConcertSessionTreeItem */
-class CONCERTSHAREDSLATE_API FConcertGroupTreeItem : public FConcertTreeItem
+class FConcertGroupTreeItem : public FConcertTreeItem
 {
 public:
 
@@ -34,4 +34,4 @@ protected:
 	
 	const FGetSessions GetSessionsFunc;
 	const TSet<FConcertSessionTreeItem::EType> AllowedTypes;
-};
\ No newline at end of file
+};
diff --git a/Engine/Plugins/Developer/Concert/ConcertUI/ConcertSharedSlate/Source/ConcertSharedSlate/Public/Session/Browser/Items/ConcertSessionTreeItem.h b/Engine/Plugins/Developer/Concert/ConcertUI/ConcertSharedSlate/Source/ConcertSharedSlate/Public/Session/Browser/Items/ConcertSessionTreeItem.h
--- a/Engine/Plugins/Developer/Concert/ConcertUI/ConcertSharedSlate/Source/ConcertSharedSlate/Public/Session/Browser/Items/ConcertSessionTreeItem.h
+++ b/Engine/Plugins/Developer/Concert/ConcertUI/ConcertSharedSlate/Source/ConcertSharedSlate/Public/Session/Browser/Items/ConcertSessionTreeItem.h
@@ -10,7 +10,7 @@
 DECLARE_MULTICAST_DELEGATE(FOnBeginEditConcertSessionNameRequest)
 
 /** An item that represents a session, i.e. has session info. */
-class CONCERTSHAREDSLATE_API FConcertSessionTreeItem : public FConcertTreeItem
+class FConcertSessionTreeItem : public FConcertTreeItem
 {
 public:
 	enum class EType
diff --git a/Engine/Plugins/Developer/Concert/ConcertUI/ConcertSharedSlate/Source/ConcertSharedSlate/Public/Session/Browser/Items/ConcertTreeItem.h b/Engine/Plugins/Developer/Concert/ConcertUI/ConcertSharedSlate/Source/ConcertSharedSlate/Public/Session/Browser/Items/ConcertTreeItem.h
--- a/Engine/Plugins/Developer/Concert/ConcertUI/ConcertSharedSlate/Source/ConcertSharedSlate/Public/Session/Browser/Items/ConcertTreeItem.h
+++ b/Engine/Plugins/Developer/Concert/ConcertUI/ConcertSharedSlate/Source/ConcertSharedSlate/Public/Session/Browser/Items/ConcertTreeItem.h
@@ -5,11 +5,11 @@
 #include "CoreMinimal.h"
 
 /** Base class for everything displayed in session view */
-class CONCERTSHAREDSLATE_API FConcertTreeItem
+class FConcertTreeItem
 {
 public:
 
 	virtual void GetChildren(TArray<TSharedPtr<FConcertTreeItem>>& OutChildren) const = 0;
 	
 	virtual ~FConcertTreeItem() = default;
-};
\ No newline at end of file
+};
diff --git a/Engine/Plugins/Developer/Concert/ConcertUI/ConcertSharedSlate/Source/ConcertSharedSlate/Public/Session/Browser/SConcertSessionBrowser.h b/Engine/Plugins/Developer/Concert/ConcertUI/ConcertSharedSlate/Source/ConcertSharedSlate/Public/Session/Browser/SConcertSessionBrowser.h
--- a/Engine/Plugins/Developer/Concert/ConcertUI/ConcertSharedSlate/Source/ConcertSharedSlate/Public/Session/Browser/SConcertSessionBrowser.h
+++ b/Engine/Plugins/Developer/Concert/ConcertUI/ConcertSharedSlate/Source/ConcertSharedSlate/Public/Session/Browser/SConcertSessionBrowser.h
@@ -12,6 +12,8 @@
 #include "Widgets/Views/SHeaderRow.h"
 #include "Widgets/Views/STreeView.h"
 
+#define UE_API CONCERTSHAREDSLATE_API
+
 class FConcertActiveGroupTreeItem;
 class FConcertArchivedGroupTreeItem;
 class FConcertTreeItem;
@@ -37,24 +39,24 @@ DECLARE_DELEGATE_RetVal_OneParam(bool, FCanRemoveSessions, const TArray<TSharedP
  * Enables the user to browse/search/filter/sort active and archived sessions, create new session,
  * archive active sessions, restore archived sessions, join a session and open the settings dialog.
  */
-class CONCERTSHAREDSLATE_API SConcertSessionBrowser : public SCompoundWidget
+class SConcertSessionBrowser : public SCompoundWidget
 {
 public:
 	
-	struct CONCERTSHAREDSLATE_API ControlButtonExtensionHooks
+	struct ControlButtonExtensionHooks
 	{
 		/** Contains: Create Session */
-		static const FName BeforeSeparator;
+		static UE_API const FName BeforeSeparator;
 		/** Just separates the two */
-		static const FName Separator;
+		static UE_API const FName Separator;
 		/** Contains: Restore, Archive, Delete */
-		static const FName AfterSeparator;
+		static UE_API const FName AfterSeparator;
 	};
 
-	struct CONCERTSHAREDSLATE_API SessionContextMenuExtensionHooks
+	struct SessionContextMenuExtensionHooks
 	{
 		/** Contains: Archive (ActiveSession), Restore (ArchivedSession), Rename, Delete  */
-		static const FName ManageSession;
+		static UE_API const FName ManageSession;
 	};
 	
 	SLATE_BEGIN_ARGS(SConcertSessionBrowser) { }
@@ -98,16 +100,16 @@ public:
 	* @param InController The controller used to send queries from the UI - represents controller in mode-view-controller pattern.
 	* @param[in,out] InSearchText The text to set in the search box and to remember (as output). Cannot be null.
 	*/
-	void Construct(const FArguments& InArgs, TSharedRef<IConcertSessionBrowserController> InController, TSharedPtr<FText> InSearchText);
+	UE_API void Construct(const FArguments& InArgs, TSharedRef<IConcertSessionBrowserController> InController, TSharedPtr<FText> InSearchText);
 
 	//~ Begin SWidget Interface
-	virtual void Tick(const FGeometry& AllottedGeometry, const double InCurrentTime, const float InDeltaTime) override;
-	virtual FReply OnKeyDown(const FGeometry& MyGeometry, const FKeyEvent& InKeyEvent) override;
+	UE_API virtual void Tick(const FGeometry& AllottedGeometry, const double InCurrentTime, const float InDeltaTime) override;
+	UE_API virtual FReply OnKeyDown(const FGeometry& MyGeometry, const FKeyEvent& InKeyEvent) override;
 	//~ End SWidget Interface
 	
-	void RefreshSessionList();
+	UE_API void RefreshSessionList();
 	bool HasAnySessions() const { return Sessions.Num() > 0; }
-	TArray<TSharedPtr<FConcertSessionTreeItem>> GetSelectedItems() const;
+	UE_API TArray<TSharedPtr<FConcertSessionTreeItem>> GetSelectedItems() const;
 
 	bool IsNewButtonEnabled() const { return IsNewButtonEnabledInternal(); }
 	bool IsRestoreButtonEnabled() const { return IsRestoreButtonEnabledInternal(); }
@@ -129,83 +131,83 @@ private:
 	};
 	
 	// Layout the 'session|details' split view.
-	TSharedRef<SWidget> MakeBrowserContent(const FArguments& InArgs);
+	UE_API TSharedRef<SWidget> MakeBrowserContent(const FArguments& InArgs);
 	
-	void OnGetChildren(TSharedPtr<FConcertTreeItem> InItem, TArray<TSharedPtr<FConcertTreeItem>>& OutChildren);
+	UE_API void OnGetChildren(TSharedPtr<FConcertTreeItem> InItem, TArray<TSharedPtr<FConcertTreeItem>>& OutChildren);
 
 	// Layout the sessions view and controls.
-	TSharedRef<SWidget> MakeControlBar(const FArguments& InArgs);
-	TSharedRef<SWidget> MakeButtonBar(const FArguments& InArgs);
-	TSharedRef<SWidget> MakeSessionTableView(const FArguments& InArgs);
-	TSharedRef<SWidget> MakeSessionTableFooter();
-	TSharedRef<ITableRow> OnGenerateSessionRowWidget(TSharedPtr<FConcertTreeItem> Item, const TSharedRef<STableViewBase>& OwnerTable);
+	UE_API TSharedRef<SWidget> MakeControlBar(const FArguments& InArgs);
+	UE_API TSharedRef<SWidget> MakeButtonBar(const FArguments& InArgs);
+	UE_API TSharedRef<SWidget> MakeSessionTableView(const FArguments& InArgs);
+	UE_API TSharedRef<SWidget> MakeSessionTableFooter();
+	UE_API TSharedRef<ITableRow> OnGenerateSessionRowWidget(TSharedPtr<FConcertTreeItem> Item, const TSharedRef<STableViewBase>& OwnerTable);
 
 	// Creates row widgets for session list view, validates user inputs and forward user requests for processing to a delegate function implemented by this class.
-	TSharedRef<ITableRow> MakeActiveSessionRowWidget(const TSharedPtr<FConcertSessionTreeItem>& ActiveItem, const TSharedRef<STableViewBase>& OwnerTable);
-	TSharedRef<ITableRow> MakeArchivedSessionRowWidget(const TSharedPtr<FConcertSessionTreeItem>& ArchivedItem, const TSharedRef<STableViewBase>& OwnerTable);
-	TSharedRef<ITableRow> MakeNewSessionRowWidget(const TSharedPtr<FConcertSessionTreeItem>& NewItem, const TSharedRef<STableViewBase>& OwnerTable);
-	TSharedRef<ITableRow> MakeRestoreSessionRowWidget(const TSharedPtr<FConcertSessionTreeItem>& RestoreItem, const TSharedRef<STableViewBase>& OwnerTable);
-	TSharedRef<ITableRow> MakeSaveSessionRowWidget(const TSharedPtr<FConcertSessionTreeItem>& ArchivedItem, const TSharedRef<STableViewBase>& OwnerTable);
-	TSharedRef<ITableRow> MakeGroupRowWidget(const TSharedPtr<FConcertTreeItem>& GroupItem, const TSharedRef<STableViewBase>& OwnerTable); 
+	UE_API TSharedRef<ITableRow> MakeActiveSessionRowWidget(const TSharedPtr<FConcertSessionTreeItem>& ActiveItem, const TSharedRef<STableViewBase>& OwnerTable);
+	UE_API TSharedRef<ITableRow> MakeArchivedSessionRowWidget(const TSharedPtr<FConcertSessionTreeItem>& ArchivedItem, const TSharedRef<STableViewBase>& OwnerTable);
+	UE_API TSharedRef<ITableRow> MakeNewSessionRowWidget(const TSharedPtr<FConcertSessionTreeItem>& NewItem, const TSharedRef<STableViewBase>& OwnerTable);
+	UE_API TSharedRef<ITableRow> MakeRestoreSessionRowWidget(const TSharedPtr<FConcertSessionTreeItem>& RestoreItem, const TSharedRef<STableViewBase>& OwnerTable);
+	UE_API TSharedRef<ITableRow> MakeSaveSessionRowWidget(const TSharedPtr<FConcertSessionTreeItem>& ArchivedItem, const TSharedRef<STableViewBase>& OwnerTable);
+	UE_API TSharedRef<ITableRow> MakeGroupRowWidget(const TSharedPtr<FConcertTreeItem>& GroupItem, const TSharedRef<STableViewBase>& OwnerTable); 
 
 	// Creates the contextual menu when right clicking a session list view row.
-	TSharedPtr<SWidget> MakeTableContextualMenu();
-	TSharedRef<SWidget> MakeViewOptionsComboButtonMenu();
+	UE_API TSharedPtr<SWidget> MakeTableContextualMenu();
+	UE_API TSharedRef<SWidget> MakeViewOptionsComboButtonMenu();
 
 	// The buttons above the session view.
-	bool IsNewButtonEnabledInternal() const;
-	bool IsRestoreButtonEnabledInternal() const;
-	bool IsArchiveButtonEnabledInternal() const;
-	bool IsRenameButtonEnabledInternal() const;
-	bool IsDeleteButtonEnabledInternal() const;
-	FReply OnNewButtonClicked();
-	FReply OnRestoreButtonClicked();
-	FReply OnArchiveButtonClicked();
-	FReply OnDeleteButtonClicked();
-	void OnBeginEditingSessionName(TSharedPtr<FConcertSessionTreeItem> Item);
+	UE_API bool IsNewButtonEnabledInternal() const;
+	UE_API bool IsRestoreButtonEnabledInternal() const;
+	UE_API bool IsArchiveButtonEnabledInternal() const;
+	UE_API bool IsRenameButtonEnabledInternal() const;
+	UE_API bool IsDeleteButtonEnabledInternal() const;
+	UE_API FReply OnNewButtonClicked();
+	UE_API FReply OnRestoreButtonClicked();
+	UE_API FReply OnArchiveButtonClicked();
+	UE_API FReply OnDeleteButtonClicked();
+	UE_API void OnBeginEditingSessionName(TSharedPtr<FConcertSessionTreeItem> Item);
 
 	// Manipulates the sessions view (the array and the UI).
-	void OnSessionSelectionChanged(TSharedPtr<FConcertTreeItem> SelectedSession, ESelectInfo::Type SelectInfo);
-	void OnSessionSelectionChanged(TSharedPtr<FConcertSessionTreeItem> SelectedSession, ESelectInfo::Type SelectInfo);
-	void InsertNewSessionEditableRowInternal();
-	void InsertRestoreSessionAsEditableRowInternal(const TSharedPtr<FConcertSessionTreeItem>& ArchivedItem);
-	void InsertArchiveSessionAsEditableRow(const TSharedPtr<FConcertSessionTreeItem>& ActiveItem);
-	void InsertEditableSessionRow(TSharedPtr<FConcertSessionTreeItem> EditableItem, TSharedPtr<FConcertSessionTreeItem> ParentItem);
-	void RemoveSessionRow(const TSharedPtr<FConcertSessionTreeItem>& Item);
+	UE_API void OnSessionSelectionChanged(TSharedPtr<FConcertTreeItem> SelectedSession, ESelectInfo::Type SelectInfo);
+	UE_API void OnSessionSelectionChanged(TSharedPtr<FConcertSessionTreeItem> SelectedSession, ESelectInfo::Type SelectInfo);
+	UE_API void InsertNewSessionEditableRowInternal();
+	UE_API void InsertRestoreSessionAsEditableRowInternal(const TSharedPtr<FConcertSessionTreeItem>& ArchivedItem);
+	UE_API void InsertArchiveSessionAsEditableRow(const TSharedPtr<FConcertSessionTreeItem>& ActiveItem);
+	UE_API void InsertEditableSessionRow(TSharedPtr<FConcertSessionTreeItem> EditableItem, TSharedPtr<FConcertSessionTreeItem> ParentItem);
+	UE_API void RemoveSessionRow(const TSharedPtr<FConcertSessionTreeItem>& Item);
 
 	// Sessions sorting. (Sorts the session view)
-	EColumnSortMode::Type GetColumnSortMode(const FName ColumnId) const;
-	EColumnSortPriority::Type GetColumnSortPriority(const FName ColumnId) const;
-	void OnColumnSortModeChanged(const EColumnSortPriority::Type SortPriority, const FName& ColumnId, const EColumnSortMode::Type InSortMode);
-	void EnsureEditableParentChildOrder();
+	UE_API EColumnSortMode::Type GetColumnSortMode(const FName ColumnId) const;
+	UE_API EColumnSortPriority::Type GetColumnSortPriority(const FName ColumnId) const;
+	UE_API void OnColumnSortModeChanged(const EColumnSortPriority::Type SortPriority, const FName& ColumnId, const EColumnSortMode::Type InSortMode);
+	UE_API void EnsureEditableParentChildOrder();
 
 	// Sessions filtering. (Filters the session view)
-	void OnSearchTextChanged(const FText& InFilterText);
-	void OnSearchTextCommitted(const FText& InFilterText, ETextCommit::Type CommitType);
-	void OnFilterMenuChecked(const FName MenuName);
-	void PopulateSearchStrings(const FConcertSessionTreeItem& Item, TArray<FString>& OutSearchStrings) const;
-	bool IsFilteredOut(const FConcertSessionTreeItem& Item) const;
-	FText HighlightSearchText() const;
+	UE_API void OnSearchTextChanged(const FText& InFilterText);
+	UE_API void OnSearchTextCommitted(const FText& InFilterText, ETextCommit::Type CommitType);
+	UE_API void OnFilterMenuChecked(const FName MenuName);
+	UE_API void PopulateSearchStrings(const FConcertSessionTreeItem& Item, TArray<FString>& OutSearchStrings) const;
+	UE_API bool IsFilteredOut(const FConcertSessionTreeItem& Item) const;
+	UE_API FText HighlightSearchText() const;
 
 	// Passes the user requests to FConcertBrowserController.
-	void RequestCreateSession(const TSharedPtr<FConcertSessionTreeItem>& NewItem);
-	void RequestArchiveSession(const TSharedPtr<FConcertSessionTreeItem>& ActiveItem, const FString& ArchiveName);
-	void RequestRestoreSession(const TSharedPtr<FConcertSessionTreeItem>& RestoreItem, const FString& SessionName);
-	void RequestRenameSession(const TSharedPtr<FConcertSessionTreeItem>& RenamedItem, const FString& NewName);
-	void RequestDeleteSessions(const TArray<TSharedPtr<FConcertSessionTreeItem>>& ItemsToDelete);
+	UE_API void RequestCreateSession(const TSharedPtr<FConcertSessionTreeItem>& NewItem);
+	UE_API void RequestArchiveSession(const TSharedPtr<FConcertSessionTreeItem>& ActiveItem, const FString& ArchiveName);
+	UE_API void RequestRestoreSession(const TSharedPtr<FConcertSessionTreeItem>& RestoreItem, const FString& SessionName);
+	UE_API void RequestRenameSession(const TSharedPtr<FConcertSessionTreeItem>& RenamedItem, const FString& NewName);
+	UE_API void RequestDeleteSessions(const TArray<TSharedPtr<FConcertSessionTreeItem>>& ItemsToDelete);
 
-	TSharedPtr<IConcertSessionBrowserController> GetController() const;
+	UE_API TSharedPtr<IConcertSessionBrowserController> GetController() const;
 
-	bool IsGroupItem(const TSharedPtr<FConcertTreeItem>& TreeItem) const;
+	UE_API bool IsGroupItem(const TSharedPtr<FConcertTreeItem>& TreeItem) const;
 
 	// Modifying sessions arrays
-	void ResetSessions();
-	void AddSession(TSharedPtr<FConcertSessionTreeItem> Session);
-	void InsertSessionAtBeginning(TSharedPtr<FConcertSessionTreeItem> NewSession);
-	void InsertSessionAfterParent(TSharedPtr<FConcertSessionTreeItem> NewSession, TSharedPtr<FConcertSessionTreeItem> Parent);
-	EInsertPosition InsertSessionAfterParentIfAvailableOrAtBeginning(TSharedPtr<FConcertSessionTreeItem> NewSession, TSharedPtr<FConcertSessionTreeItem> Parent);
-	void RemoveSession(TSharedPtr<FConcertSessionTreeItem> Session);
-	void SortSessionList();
+	UE_API void ResetSessions();
+	UE_API void AddSession(TSharedPtr<FConcertSessionTreeItem> Session);
+	UE_API void InsertSessionAtBeginning(TSharedPtr<FConcertSessionTreeItem> NewSession);
+	UE_API void InsertSessionAfterParent(TSharedPtr<FConcertSessionTreeItem> NewSession, TSharedPtr<FConcertSessionTreeItem> Parent);
+	UE_API EInsertPosition InsertSessionAfterParentIfAvailableOrAtBeginning(TSharedPtr<FConcertSessionTreeItem> NewSession, TSharedPtr<FConcertSessionTreeItem> Parent);
+	UE_API void RemoveSession(TSharedPtr<FConcertSessionTreeItem> Session);
+	UE_API void SortSessionList();
 
 private:
 	
@@ -266,3 +268,5 @@ private:
 	TArray<TSharedPtr<FConcertSessionClientInfo>> Clients;
 	TSharedPtr<SExpandableArea> ClientsArea;
 };
+
+#undef UE_API
diff --git a/Engine/Plugins/Developer/Concert/ConcertUI/ConcertSharedSlate/Source/ConcertSharedSlate/Public/Session/History/AbstractSessionHistoryController.h b/Engine/Plugins/Developer/Concert/ConcertUI/ConcertSharedSlate/Source/ConcertSharedSlate/Public/Session/History/AbstractSessionHistoryController.h
--- a/Engine/Plugins/Developer/Concert/ConcertUI/ConcertSharedSlate/Source/ConcertSharedSlate/Public/Session/History/AbstractSessionHistoryController.h
+++ b/Engine/Plugins/Developer/Concert/ConcertUI/ConcertSharedSlate/Source/ConcertSharedSlate/Public/Session/History/AbstractSessionHistoryController.h
@@ -1,4 +1,4 @@
-﻿// Copyright Epic Games, Inc. All Rights Reserved.
+// Copyright Epic Games, Inc. All Rights Reserved.
 
 #pragma once
 
@@ -7,6 +7,8 @@
 
 #include "Async/Future.h"
 
+#define UE_API CONCERTSHAREDSLATE_API
+
 struct FConcertClientInfo;
 struct FConcertSyncTransactionEvent;
 struct FConcertSyncPackageEventMetaData;
@@ -14,15 +16,15 @@ struct FConcertSessionActivity;
 class IConcertSyncClient;
 
 /** Acts as controller of SSessionHistory in model-view-component pattern. */
-class CONCERTSHAREDSLATE_API FAbstractSessionHistoryController : public TSharedFromThis<FAbstractSessionHistoryController>
+class FAbstractSessionHistoryController : public TSharedFromThis<FAbstractSessionHistoryController>
 {
 public:
 
-	FAbstractSessionHistoryController(SSessionHistory::FArguments Arguments = SSessionHistory::FArguments());
+	UE_API FAbstractSessionHistoryController(SSessionHistory::FArguments Arguments = SSessionHistory::FArguments());
 	virtual ~FAbstractSessionHistoryController() = default;
 	
 	/** Fetches activities from the server and updates the list view. */
-	void ReloadActivities();
+	UE_API void ReloadActivities();
 	
 	const TSharedRef<SSessionHistory>& GetSessionHistory() const { return SessionHistory; }
 
@@ -42,5 +44,7 @@ private:
 	/** The widget being managed */
 	TSharedRef<SSessionHistory> SessionHistory;
 
-	TSharedRef<SSessionHistory> MakeSessionHistory(SSessionHistory::FArguments Arguments) const;
-};
\ No newline at end of file
+	UE_API TSharedRef<SSessionHistory> MakeSessionHistory(SSessionHistory::FArguments Arguments) const;
+};
+
+#undef UE_API
diff --git a/Engine/Plugins/Developer/Concert/ConcertUI/ConcertSharedSlate/Source/ConcertSharedSlate/Public/Session/History/SEditableSessionHistory.h b/Engine/Plugins/Developer/Concert/ConcertUI/ConcertSharedSlate/Source/ConcertSharedSlate/Public/Session/History/SEditableSessionHistory.h
--- a/Engine/Plugins/Developer/Concert/ConcertUI/ConcertSharedSlate/Source/ConcertSharedSlate/Public/Session/History/SEditableSessionHistory.h
+++ b/Engine/Plugins/Developer/Concert/ConcertUI/ConcertSharedSlate/Source/ConcertSharedSlate/Public/Session/History/SEditableSessionHistory.h
@@ -1,4 +1,4 @@
-﻿// Copyright Epic Games, Inc. All Rights Reserved.
+// Copyright Epic Games, Inc. All Rights Reserved.
 
 #pragma once
 
@@ -7,6 +7,8 @@
 #include "Widgets/DeclarativeSyntaxSupport.h"
 #include "Widgets/SCompoundWidget.h"
 
+#define UE_API CONCERTSHAREDSLATE_API
+
 class SSessionHistory;
 struct FConcertSessionActivity;
 
@@ -25,7 +27,7 @@ struct FCanPerformActionResult
 };
 
 /** Allows activities in the session history to be deleted. */
-class CONCERTSHAREDSLATE_API SEditableSessionHistory : public SCompoundWidget
+class SEditableSessionHistory : public SCompoundWidget
 {
 public:
 
@@ -53,10 +55,10 @@ public:
 		SLATE_NAMED_SLOT(FArguments, StatusBar)
 	SLATE_END_ARGS()
 
-	void Construct(const FArguments& InArgs);
+	UE_API void Construct(const FArguments& InArgs);
 	
 	//~ Begin SWidget Interface
-	virtual FReply OnKeyDown(const FGeometry& MyGeometry, const FKeyEvent& InKeyEvent) override;
+	UE_API virtual FReply OnKeyDown(const FGeometry& MyGeometry, const FKeyEvent& InKeyEvent) override;
 	//~ End SWidget Interface
 	
 private:
@@ -70,19 +72,21 @@ private:
 	FCanPerformActionOnActivities CanUnmuteActivitiesFunc;
 	FRequestActivitiesAction UnmuteActivitiesFunc;
 
-	TSharedPtr<SWidget> OnContextMenuOpening();
+	UE_API TSharedPtr<SWidget> OnContextMenuOpening();
 	
-	FReply OnClickDeleteActivitiesButton() const;
-	FText GetDeleteActivitiesToolTip() const;
-	bool IsDeleteButtonEnabled() const;
+	UE_API FReply OnClickDeleteActivitiesButton() const;
+	UE_API FText GetDeleteActivitiesToolTip() const;
+	UE_API bool IsDeleteButtonEnabled() const;
 	
-	FReply OnClickMuteActivitesButton() const;
-	FText GetMuteActivitiesToolTip() const;
-	bool IsMuteButtonEnabled() const;
+	UE_API FReply OnClickMuteActivitesButton() const;
+	UE_API FText GetMuteActivitiesToolTip() const;
+	UE_API bool IsMuteButtonEnabled() const;
 	
-	FReply OnClickUnmuteActivitesButton() const;
-	FText GetUnmuteActivitiesToolTip() const;
-	bool IsUnmuteButtonEnabled() const;
+	UE_API FReply OnClickUnmuteActivitesButton() const;
+	UE_API FText GetUnmuteActivitiesToolTip() const;
+	UE_API bool IsUnmuteButtonEnabled() const;
 
-	FText GenerateTooltip(const FCanPerformActionOnActivities& CanPerformAction, FText SelectActivityToolTip, FText PerformActionToolTipFmt, FText CannotPerformActionToolTipFmt) const;
+	UE_API FText GenerateTooltip(const FCanPerformActionOnActivities& CanPerformAction, FText SelectActivityToolTip, FText PerformActionToolTipFmt, FText CannotPerformActionToolTipFmt) const;
 };
+
+#undef UE_API
diff --git a/Engine/Plugins/Developer/Concert/ConcertUI/ConcertSharedSlate/Source/ConcertSharedSlate/Public/Session/History/SSessionHistory.h b/Engine/Plugins/Developer/Concert/ConcertUI/ConcertSharedSlate/Source/ConcertSharedSlate/Public/Session/History/SSessionHistory.h
--- a/Engine/Plugins/Developer/Concert/ConcertUI/ConcertSharedSlate/Source/ConcertSharedSlate/Public/Session/History/SSessionHistory.h
+++ b/Engine/Plugins/Developer/Concert/ConcertUI/ConcertSharedSlate/Source/ConcertSharedSlate/Public/Session/History/SSessionHistory.h
@@ -8,6 +8,8 @@
 #include "Widgets/DeclarativeSyntaxSupport.h"
 #include "Widgets/SCompoundWidget.h"
 
+#define UE_API CONCERTSHAREDSLATE_API
+
 namespace UE::ConcertSharedSlate
 {
 	class IConcertReflectionDataProvider;
@@ -17,7 +19,7 @@ class SSearchBox;
 class FConcertSessionActivitiesOptions;
 
 /** Wraps a session table list with a search box and a status bar for visibility settings (show connection activities, etc.). */
-class CONCERTSHAREDSLATE_API SSessionHistory : public SCompoundWidget
+class SSessionHistory : public SCompoundWidget
 {
 public:
 
@@ -82,32 +84,32 @@ public:
 	 *
 	 * @param InArgs The Slate argument list.
 	 */
-	void Construct(const FArguments& InArgs);
+	UE_API void Construct(const FArguments& InArgs);
 
 	/** Fetches activities from the server and updates the list view. */
-	void ReloadActivities(TMap<FGuid, FConcertClientInfo> EndpointClientInfoMap, TArray<FConcertSessionActivity> FetchedActivities);
+	UE_API void ReloadActivities(TMap<FGuid, FConcertClientInfo> EndpointClientInfoMap, TArray<FConcertSessionActivity> FetchedActivities);
 	
 	/** Callback for handling the a new or updated activity item. */ 
-	void HandleActivityAddedOrUpdated(const FConcertClientInfo& InClientInfo, const FConcertSyncActivity& InActivity, const FStructOnScope& InActivitySummary);
+	UE_API void HandleActivityAddedOrUpdated(const FConcertClientInfo& InClientInfo, const FConcertSyncActivity& InActivity, const FStructOnScope& InActivitySummary);
 
 	bool IsLastColumn(FName ColumnId) const { return ActivityListView->IsLastColumn(ColumnId); }
 	
-	TSet<TSharedRef<FConcertSessionActivity>> GetSelectedActivities() const;
+	UE_API TSet<TSharedRef<FConcertSessionActivity>> GetSelectedActivities() const;
 	const TArray<TSharedPtr<FConcertSessionActivity>>& GetActivities() const { return ActivityListView->GetActivities(); }
 	void SetSelectedActivities(const TArray<TSharedPtr<FConcertSessionActivity>>& ActivitiesToSelect) { ActivityListView->SetSelectedActivities(ActivitiesToSelect); }
 
-	void OnColumnVisibilitySettingsChanged(const FColumnVisibilitySnapshot& ColumnSnapshot);
+	UE_API void OnColumnVisibilitySettingsChanged(const FColumnVisibilitySnapshot& ColumnSnapshot);
 
 private:
 
 	/** Invoked when the text in the search box widget changes. */
-	void OnSearchTextChanged(const FText& InSearchText);
+	UE_API void OnSearchTextChanged(const FText& InSearchText);
 
 	/** Invoked when the text in the search box widget is committed. */
-	void OnSearchTextCommitted(const FText& InFilterText, ETextCommit::Type CommitType);
+	UE_API void OnSearchTextCommitted(const FText& InFilterText, ETextCommit::Type CommitType);
 
 	/** Returns the text to highlight when the search bar has a text set. */
-	FText HighlightSearchedText() const;
+	UE_API FText HighlightSearchedText() const;
 
 private:
 	
@@ -132,5 +134,7 @@ private:
 	/** Used to limit activities. */
 	FAllowActivity AllowActivityFunc;
 
-	TOptional<FConcertClientInfo> GetClientInfo(FGuid Guid) const;
-};
\ No newline at end of file
+	UE_API TOptional<FConcertClientInfo> GetClientInfo(FGuid Guid) const;
+};
+
+#undef UE_API
diff --git a/Engine/Plugins/Developer/Concert/ConcertUI/ConcertSharedSlate/Source/ConcertSharedSlate/Public/Session/History/SSessionHistoryWrapper.h b/Engine/Plugins/Developer/Concert/ConcertUI/ConcertSharedSlate/Source/ConcertSharedSlate/Public/Session/History/SSessionHistoryWrapper.h
--- a/Engine/Plugins/Developer/Concert/ConcertUI/ConcertSharedSlate/Source/ConcertSharedSlate/Public/Session/History/SSessionHistoryWrapper.h
+++ b/Engine/Plugins/Developer/Concert/ConcertUI/ConcertSharedSlate/Source/ConcertSharedSlate/Public/Session/History/SSessionHistoryWrapper.h
@@ -8,18 +8,22 @@
 #include "Widgets/SCompoundWidget.h"
 #include "Widgets/DeclarativeSyntaxSupport.h"
 
+#define UE_API CONCERTSHAREDSLATE_API
+
 /** Shows a SSessionHistory widget - keeping alive the controller as long as this widget lives. Useful for closeable tabs. */
-class CONCERTSHAREDSLATE_API SSessionHistoryWrapper : public SCompoundWidget
+class SSessionHistoryWrapper : public SCompoundWidget
 {
 public:
 
 	SLATE_BEGIN_ARGS(SSessionHistoryWrapper) {}
 	SLATE_END_ARGS()
 
-	void Construct(const FArguments& InArgs, TSharedRef<FAbstractSessionHistoryController> InController);
+	UE_API void Construct(const FArguments& InArgs, TSharedRef<FAbstractSessionHistoryController> InController);
 
 private:
 
 	/** Manages SessionHistory */
 	TSharedPtr<FAbstractSessionHistoryController> Controller;
 };
+
+#undef UE_API
diff --git a/Engine/Plugins/Developer/Concert/ConcertUI/ConcertSharedSlate/Source/ConcertSharedSlate/Public/Widgets/Client/SClientName.h b/Engine/Plugins/Developer/Concert/ConcertUI/ConcertSharedSlate/Source/ConcertSharedSlate/Public/Widgets/Client/SClientName.h
--- a/Engine/Plugins/Developer/Concert/ConcertUI/ConcertSharedSlate/Source/ConcertSharedSlate/Public/Widgets/Client/SClientName.h
+++ b/Engine/Plugins/Developer/Concert/ConcertUI/ConcertSharedSlate/Source/ConcertSharedSlate/Public/Widgets/Client/SClientName.h
@@ -7,6 +7,8 @@
 #include "Widgets/DeclarativeSyntaxSupport.h"
 #include "Widgets/SCompoundWidget.h"
 
+#define UE_API CONCERTSHAREDSLATE_API
+
 class IConcertClient;
 
 /** Contains SClientName's parentheses content definitions. */
@@ -28,7 +30,7 @@ namespace UE::ConcertSharedSlate
 	 *  - DisplayName is FConcertClientInfo::DisplayName
 	 *  - ParenthesesContent is additional info you can supply, like "You" (optional). See ParenthesesClientNameContent. 
 	 */
-	class CONCERTSHAREDSLATE_API SClientName : public SCompoundWidget
+	class SClientName : public SCompoundWidget
 	{
 	public:
 
@@ -50,12 +52,12 @@ namespace UE::ConcertSharedSlate
 			SLATE_ARGUMENT(FSlateFontInfo, Font)
 		SLATE_END_ARGS()
 
-		void Construct(const FArguments& InArgs);
+		UE_API void Construct(const FArguments& InArgs);
 
 		/** @return The display that would be used given the settings. */
-		static FText GetDisplayText(const FConcertClientInfo& Info, bool bDisplayAsLocalClient);
+		static UE_API FText GetDisplayText(const FConcertClientInfo& Info, bool bDisplayAsLocalClient);
 		/** @return The display that would be used given the settings. */
-		static FText GetDisplayText(const FConcertClientInfo& Info, const FText& ParenthesesContent = FText::GetEmpty());
+		static UE_API FText GetDisplayText(const FConcertClientInfo& Info, const FText& ParenthesesContent = FText::GetEmpty());
 
 	private:
 
@@ -65,8 +67,10 @@ namespace UE::ConcertSharedSlate
 		TAttribute<FText> ParenthesisContentAttribute;
 
 		/** Gets the display name. */
-		FText GetClientDisplayName() const;
+		UE_API FText GetClientDisplayName() const;
 		/** Gets the avatar color */
-		FSlateColor GetAvatarColor() const;
+		UE_API FSlateColor GetAvatarColor() const;
 	};
 }
+
+#undef UE_API
diff --git a/Engine/Plugins/Developer/Concert/ConcertUI/ConcertSharedSlate/Source/ConcertSharedSlate/Public/Widgets/Client/SHorizontalClientList.h b/Engine/Plugins/Developer/Concert/ConcertUI/ConcertSharedSlate/Source/ConcertSharedSlate/Public/Widgets/Client/SHorizontalClientList.h
--- a/Engine/Plugins/Developer/Concert/ConcertUI/ConcertSharedSlate/Source/ConcertSharedSlate/Public/Widgets/Client/SHorizontalClientList.h
+++ b/Engine/Plugins/Developer/Concert/ConcertUI/ConcertSharedSlate/Source/ConcertSharedSlate/Public/Widgets/Client/SHorizontalClientList.h
@@ -8,6 +8,8 @@
 #include "Widgets/DeclarativeSyntaxSupport.h"
 #include "Widgets/SCompoundWidget.h"
 
+#define UE_API CONCERTSHAREDSLATE_API
+
 class IConcertClient;
 class SScrollBox;
 class SWidgetSwitcher;
@@ -18,12 +20,12 @@ struct FGuid;
 namespace UE::ConcertSharedSlate
 {
 	/** Aligns client widgets from left to right. If there is not enough space, a horizontal scroll bar cuts of the list. */
-	class CONCERTSHAREDSLATE_API SHorizontalClientList : public SCompoundWidget
+	class SHorizontalClientList : public SCompoundWidget
 	{
 	public:
 
 		/** @return The display string a SHorizontalClientList would display with the given state. Returns unset optional if EmptyListSlot would be shown. */
-		static TOptional<FString> GetDisplayString(
+		static UE_API TOptional<FString> GetDisplayString(
 			const TConstArrayView<FGuid>& Clients,
 			const FGetOptionalClientInfo& GetClientInfoDelegate,
 			const FClientSortPredicate& SortPredicate,
@@ -57,10 +59,10 @@ namespace UE::ConcertSharedSlate
 			SLATE_NAMED_SLOT(FArguments, EmptyListSlot)
 		SLATE_END_ARGS()
 
-		void Construct(const FArguments& InArgs);
+		UE_API void Construct(const FArguments& InArgs);
 
 		/** Refreshes the list. */
-		void RefreshList(const TConstArrayView<FGuid>& Clients) const;
+		UE_API void RefreshList(const TConstArrayView<FGuid>& Clients) const;
 
 	private:
 		
@@ -84,4 +86,6 @@ namespace UE::ConcertSharedSlate
 		/** Contains the children. */
 		TSharedPtr<SScrollBox> ScrollBox;
 	};
-}
\ No newline at end of file
+}
+
+#undef UE_API
diff --git a/Engine/Plugins/Developer/Concert/ConcertUI/ConcertSharedSlate/Source/ConcertSharedSlate/Public/Widgets/Client/SLocalClientName.h b/Engine/Plugins/Developer/Concert/ConcertUI/ConcertSharedSlate/Source/ConcertSharedSlate/Public/Widgets/Client/SLocalClientName.h
--- a/Engine/Plugins/Developer/Concert/ConcertUI/ConcertSharedSlate/Source/ConcertSharedSlate/Public/Widgets/Client/SLocalClientName.h
+++ b/Engine/Plugins/Developer/Concert/ConcertUI/ConcertSharedSlate/Source/ConcertSharedSlate/Public/Widgets/Client/SLocalClientName.h
@@ -8,6 +8,8 @@
 #include "Widgets/DeclarativeSyntaxSupport.h"
 #include "Widgets/SCompoundWidget.h"
 
+#define UE_API CONCERTSHAREDSLATE_API
+
 class IConcertClient;
 
 namespace UE::ConcertSharedSlate
@@ -18,7 +20,7 @@ namespace UE::ConcertSharedSlate
 	 * The name will look like "Client Name (me)".
 	 * @see SRemoteClientName
 	 */
-	class CONCERTSHAREDSLATE_API SLocalClientName : public SCompoundWidget
+	class SLocalClientName : public SCompoundWidget
 	{
 	public:
 
@@ -37,6 +39,8 @@ namespace UE::ConcertSharedSlate
 			SLATE_ARGUMENT(FSlateFontInfo, Font)
 		SLATE_END_ARGS()
 
-		void Construct(const FArguments& InArgs);
+		UE_API void Construct(const FArguments& InArgs);
 	};
 }
+
+#undef UE_API
diff --git a/Engine/Plugins/Developer/Concert/ConcertUI/ConcertSharedSlate/Source/ConcertSharedSlate/Public/Widgets/Client/SRemoteClientName.h b/Engine/Plugins/Developer/Concert/ConcertUI/ConcertSharedSlate/Source/ConcertSharedSlate/Public/Widgets/Client/SRemoteClientName.h
--- a/Engine/Plugins/Developer/Concert/ConcertUI/ConcertSharedSlate/Source/ConcertSharedSlate/Public/Widgets/Client/SRemoteClientName.h
+++ b/Engine/Plugins/Developer/Concert/ConcertUI/ConcertSharedSlate/Source/ConcertSharedSlate/Public/Widgets/Client/SRemoteClientName.h
@@ -8,6 +8,8 @@
 #include "Widgets/DeclarativeSyntaxSupport.h"
 #include "Widgets/SCompoundWidget.h"
 
+#define UE_API CONCERTSHAREDSLATE_API
+
 class IConcertClient;
 
 namespace UE::ConcertSharedSlate
@@ -21,7 +23,7 @@ namespace UE::ConcertSharedSlate
 	 * If the client disconnects, the last known info is used.
 	 * If the client info is unknown, the widget will display an empty FConcertClientInfo;
 	 */
-	class CONCERTSHAREDSLATE_API SRemoteClientName : public SCompoundWidget
+	class SRemoteClientName : public SCompoundWidget
 	{
 	public:
 
@@ -40,7 +42,7 @@ namespace UE::ConcertSharedSlate
 			SLATE_ARGUMENT(FSlateFontInfo, Font)
 		SLATE_END_ARGS()
 
-		void Construct(const FArguments& InArgs);
+		UE_API void Construct(const FArguments& InArgs);
 
 	private:
 
@@ -53,6 +55,8 @@ namespace UE::ConcertSharedSlate
 		 */
 		mutable TOptional<FConcertClientInfo> LastKnownClientInfo;
 
-		TOptional<FConcertClientInfo> GetClientInfo() const;
+		UE_API TOptional<FConcertClientInfo> GetClientInfo() const;
 	};
 }
+
+#undef UE_API
